[LLM] restoring LLM tracking for lost during voice chat refactoring, using new AudioVoiceChat tag.

diff --git a/Engine/Plugins/Online/VoiceChat/VivoxVoiceChat/Source/Private/VivoxVoiceChat.cpp b/Engine/Plugins/Online/VoiceChat/VivoxVoiceChat/Source/Private/VivoxVoiceChat.cpp
--- a/Engine/Plugins/Online/VoiceChat/VivoxVoiceChat/Source/Private/VivoxVoiceChat.cpp
+++ b/Engine/Plugins/Online/VoiceChat/VivoxVoiceChat/Source/Private/VivoxVoiceChat.cpp
@@ -7,6 +7,7 @@
 #include "Misc/ConfigCacheIni.h"
 #include "Misc/ScopeLock.h"
 #include "Modules/ModuleManager.h"
+#include "HAL/LowLevelMemTracker.h"
 
 #include "Vxc.h"
 #include "VxcErrors.h"
@@ -119,31 +120,37 @@ FVivoxVoiceChat::~FVivoxVoiceChat()
 
 static void* VivoxMalloc(size_t bytes)
 {
+	LLM_SCOPE( ELLMTag::AudioVoiceChat );
 	return FMemory::Malloc(bytes);
 }
 
 static void VivoxFree(void* ptr)
 {
+	LLM_SCOPE( ELLMTag::AudioVoiceChat );
 	FMemory::Free(ptr);
 }
 
 static void* VivoxRealloc(void* ptr, size_t bytes)
 {
+	LLM_SCOPE( ELLMTag::AudioVoiceChat );
 	return FMemory::Realloc(ptr, bytes);
 }
 
 static void* VivoxCalloc(size_t num, size_t bytes)
 {
+	LLM_SCOPE( ELLMTag::AudioVoiceChat );
 	return FMemory::Malloc(bytes*num);
 }
 
 static void* VivoxMallocAligned(size_t alignment, size_t bytes)
 {
+	LLM_SCOPE( ELLMTag::AudioVoiceChat );
 	return FMemory::Malloc(bytes, alignment);
 }
 
 static void VivoxFreeAligned(void* ptr)
 {
+	LLM_SCOPE( ELLMTag::AudioVoiceChat );
 	FMemory::Free(ptr);
 }
 
