Add CallStats callback to IVoiceChat, and Vivox implementation

diff --git a/Engine/Plugins/Online/VoiceChat/VivoxVoiceChat/Source/Private/VivoxVoiceChat.cpp b/Engine/Plugins/Online/VoiceChat/VivoxVoiceChat/Source/Private/VivoxVoiceChat.cpp
--- a/Engine/Plugins/Online/VoiceChat/VivoxVoiceChat/Source/Private/VivoxVoiceChat.cpp
+++ b/Engine/Plugins/Online/VoiceChat/VivoxVoiceChat/Source/Private/VivoxVoiceChat.cpp
@@ -1281,7 +1281,18 @@ void FVivoxVoiceChat::onChannelExited(const VivoxClientApi::AccountName& Account
 
 void FVivoxVoiceChat::onCallStatsUpdated(const VivoxClientApi::AccountName& AccountName, vx_call_stats_t& Stats, bool bIsFinal)
 {
-
+	const int TotalPacketsLost = Stats.incoming_packetloss + Stats.incoming_discarded + Stats.incoming_out_of_time;
+	const int TotalPackets = TotalPacketsLost + Stats.incoming_received;
+
+	FVoiceChatCallStats CallStats;
+	CallStats.CallLength = Stats.sample_interval_end - Stats.sample_interval_begin;
+	CallStats.LatencyMinMeasuredSeconds = Stats.min_latency;
+	CallStats.LatencyMaxMeasuredSeconds = Stats.max_latency;
+	CallStats.LatencyAverageMeasuredSeconds = Stats.latency_measurement_count > 0 ? Stats.latency_sum / Stats.latency_measurement_count : 0.0;
+	CallStats.PacketsNumLost = TotalPacketsLost;
+	CallStats.PacketsNumTotal = TotalPackets;
+
+	OnVoiceChatCallStatsUpdatedDelegate.Broadcast(CallStats);
 }
 
 void FVivoxVoiceChat::onParticipantAdded(const VivoxClientApi::AccountName& AccountName, const VivoxClientApi::Uri& ChannelUri, const VivoxClientApi::Uri& ParticipantUri, bool bIsLoggedInUser)
diff --git a/Engine/Plugins/Online/VoiceChat/VivoxVoiceChat/Source/Public/VivoxVoiceChat.h b/Engine/Plugins/Online/VoiceChat/VivoxVoiceChat/Source/Public/VivoxVoiceChat.h
--- a/Engine/Plugins/Online/VoiceChat/VivoxVoiceChat/Source/Public/VivoxVoiceChat.h
+++ b/Engine/Plugins/Online/VoiceChat/VivoxVoiceChat/Source/Public/VivoxVoiceChat.h
@@ -66,6 +66,7 @@ public:
 	virtual void JoinChannel(const FString& ChannelName, const FString& ChannelCredentials, EVoiceChatChannelType ChannelType, const FOnVoiceChatChannelJoinCompleteDelegate& Delegate, TOptional<FVoiceChatChannel3dProperties> Channel3dProperties = TOptional<FVoiceChatChannel3dProperties>()) override;
 	virtual void LeaveChannel(const FString& Channel, const FOnVoiceChatChannelLeaveCompleteDelegate& Delegate) override;
 	virtual FOnVoiceChatChannelExitedDelegate& OnVoiceChatChannelExited() override { return OnVoiceChatChannelExitedDelegate; }
+	virtual FOnVoiceChatCallStatsUpdatedDelegate& OnVoiceChatCallStatsUpdated() override { return OnVoiceChatCallStatsUpdatedDelegate; }
 	virtual void Set3DPosition(const FString& ChannelName, const FVector& SpeakerPosition, const FVector& ListenerPosition, const FVector& ListenerForwardDirection, const FVector& ListenerUpDirection) override;
 	virtual TArray<FString> GetChannels() const override;
 	virtual TArray<FString> GetPlayersInChannel(const FString& ChannelName) const override;
@@ -225,6 +226,7 @@ protected:
 	FOnVoiceChatPlayerMuteUpdatedDelegate OnVoiceChatPlayerMuteUpdatedDelegate;
 	FOnVoiceChatPlayerVolumeUpdatedDelegate OnVoiceChatPlayerVolumeUpdatedDelegate;
 	FOnVoiceChatPlayerRemovedDelegate OnVoiceChatPlayerRemovedDelegate;
+	FOnVoiceChatCallStatsUpdatedDelegate OnVoiceChatCallStatsUpdatedDelegate;
 
 	// Recording Delegates and Critical sections
 	FCriticalSection AudioRecordLock;
diff --git a/Engine/Plugins/Online/VoiceChat/VivoxVoiceChat/Source/Public/VoiceChat.h b/Engine/Plugins/Online/VoiceChat/VivoxVoiceChat/Source/Public/VoiceChat.h
--- a/Engine/Plugins/Online/VoiceChat/VivoxVoiceChat/Source/Public/VoiceChat.h
+++ b/Engine/Plugins/Online/VoiceChat/VivoxVoiceChat/Source/Public/VoiceChat.h
@@ -60,6 +60,22 @@ enum class EVoiceChatTransmitMode
 	Channel
 };
 
+struct FVoiceChatCallStats
+{
+	/** Call length in seconds */
+	double CallLength;
+	/** Minimum measured network latency in seconds. Zero if no measurements made */
+	double LatencyMinMeasuredSeconds;
+	/** Maximum measured network latency in seconds. Zero if no measurements made */
+	double LatencyMaxMeasuredSeconds;
+	/** Average measured network latency in seconds. Zero if no measurements made */
+	double LatencyAverageMeasuredSeconds;
+	/** Total number of packets, both received and lost */
+	int PacketsNumTotal;
+	/** Total number of lost packets */
+	int PacketsNumLost;
+};
+
 DECLARE_DELEGATE_OneParam(FOnVoiceChatConnectCompleteDelegate, const FVoiceChatResult& /* Result */);
 DECLARE_DELEGATE_OneParam(FOnVoiceChatDisconnectCompleteDelegate, const FVoiceChatResult& /* Result */);
 DECLARE_DELEGATE_TwoParams(FOnVoiceChatLoginCompleteDelegate, const FString& /* PlayerName */, const FVoiceChatResult& /* Result */);
@@ -68,8 +84,9 @@ DECLARE_DELEGATE_TwoParams(FOnVoiceChatChannelJoinCompleteDelegate, const FStrin
 DECLARE_DELEGATE_TwoParams(FOnVoiceChatChannelLeaveCompleteDelegate, const FString& /* ChannelName */, const FVoiceChatResult& /* Result */);
 
 DECLARE_MULTICAST_DELEGATE(FOnVoiceChatAvailableAudioDevicesChangedDelegate);
-DECLARE_MULTICAST_DELEGATE_OneParam(FOnVoiceChatDisconnectedDelegate, const FVoiceChatResult& /* Reason */);
 DECLARE_MULTICAST_DELEGATE(FOnVoiceChatReconnectedDelegate);
+DECLARE_MULTICAST_DELEGATE_OneParam(FOnVoiceChatDisconnectedDelegate, const FVoiceChatResult& /* Reason */);
+DECLARE_MULTICAST_DELEGATE_OneParam(FOnVoiceChatCallStatsUpdatedDelegate, const FVoiceChatCallStats& /* CallStats */);
 DECLARE_MULTICAST_DELEGATE_TwoParams(FOnVoiceChatChannelExitedDelegate, const FString& /* ChannelName */, const FVoiceChatResult& /* Reason */);
 DECLARE_MULTICAST_DELEGATE_TwoParams(FOnVoiceChatPlayerAddedDelegate, const FString& /* ChannelName */, const FString& /* PlayerName */);
 DECLARE_MULTICAST_DELEGATE_TwoParams(FOnVoiceChatPlayerRemovedDelegate, const FString& /* ChannelName */, const FString& /* PlayerName */);
@@ -339,6 +356,11 @@ public:
 	 */
 	virtual FOnVoiceChatChannelExitedDelegate& OnVoiceChatChannelExited() = 0;
 
+	/**
+	 * Delegate triggered when a call is ended, providing the stats for the call.
+	 */
+	virtual FOnVoiceChatCallStatsUpdatedDelegate& OnVoiceChatCallStatsUpdated() = 0;
+
 	/**
 	 * Set the 3d position of the player
 	 *
