[VivoxVoiceChat] Set new state before calling callback in onDisconnected

diff --git a/Engine/Plugins/Online/VoiceChat/VivoxVoiceChat/Source/Private/VivoxVoiceChat.cpp b/Engine/Plugins/Online/VoiceChat/VivoxVoiceChat/Source/Private/VivoxVoiceChat.cpp
--- a/Engine/Plugins/Online/VoiceChat/VivoxVoiceChat/Source/Private/VivoxVoiceChat.cpp
+++ b/Engine/Plugins/Online/VoiceChat/VivoxVoiceChat/Source/Private/VivoxVoiceChat.cpp
@@ -1148,7 +1148,10 @@ void FVivoxVoiceChat::onDisconnected(const VivoxClientApi::Uri& Server, const Vi
 
 	ClearLoginSession();
 
-	if (ConnectionState == EConnectionState::Disconnecting)
+	EConnectionState PreviousConnectionState = ConnectionState;
+	ConnectionState = EConnectionState::Disconnected;
+
+	if (PreviousConnectionState == EConnectionState::Disconnecting)
 	{
 		TriggerCompletionDelegates(OnVoiceChatDisconnectCompleteDelegates, ResultSuccess);
 	}
@@ -1156,8 +1159,6 @@ void FVivoxVoiceChat::onDisconnected(const VivoxClientApi::Uri& Server, const Vi
 	{
 		OnVoiceChatDisconnectedDelegate.Broadcast(ResultFromVivoxStatus(Status));
 	}
-
-	ConnectionState = EConnectionState::Disconnected;
 }
 
 void FVivoxVoiceChat::onLoginCompleted(const VivoxClientApi::AccountName& AccountName)
