Fix VivoxCalloc not zeroing memory

diff --git a/Engine/Plugins/Online/VoiceChat/VivoxVoiceChat/Source/Private/VivoxVoiceChat.cpp b/Engine/Plugins/Online/VoiceChat/VivoxVoiceChat/Source/Private/VivoxVoiceChat.cpp
--- a/Engine/Plugins/Online/VoiceChat/VivoxVoiceChat/Source/Private/VivoxVoiceChat.cpp
+++ b/Engine/Plugins/Online/VoiceChat/VivoxVoiceChat/Source/Private/VivoxVoiceChat.cpp
@@ -139,7 +139,10 @@ static void* VivoxRealloc(void* ptr, size_t bytes)
 static void* VivoxCalloc(size_t num, size_t bytes)
 {
 	LLM_SCOPE( ELLMTag::AudioVoiceChat );
-	return FMemory::Malloc(bytes*num);
+	const size_t Size = bytes * num;
+	void* Ret = FMemory::Malloc(Size);
+	FMemory::Memzero(Ret, Size);
+	return Ret;
 }
 
 static void* VivoxMallocAligned(size_t alignment, size_t bytes)
