Add mutex lock to vivox memory allocators for IOS.

diff --git a/Engine/Plugins/Online/VoiceChat/VivoxVoiceChat/Source/Private/VivoxVoiceChat.cpp b/Engine/Plugins/Online/VoiceChat/VivoxVoiceChat/Source/Private/VivoxVoiceChat.cpp
--- a/Engine/Plugins/Online/VoiceChat/VivoxVoiceChat/Source/Private/VivoxVoiceChat.cpp
+++ b/Engine/Plugins/Online/VoiceChat/VivoxVoiceChat/Source/Private/VivoxVoiceChat.cpp
@@ -118,27 +118,43 @@ FVivoxVoiceChat::~FVivoxVoiceChat()
 {
 }
 
+#if PLATFORM_IOS
+FCriticalSection Mutex;
+#endif
+
 static void* VivoxMalloc(size_t bytes)
 {
 	LLM_SCOPE( ELLMTag::AudioVoiceChat );
+#if PLATFORM_IOS
+	FScopeLock Lock(&Mutex);
+#endif
 	return FMemory::Malloc(bytes);
 }
 
 static void VivoxFree(void* ptr)
 {
 	LLM_SCOPE( ELLMTag::AudioVoiceChat );
+#if PLATFORM_IOS
+	FScopeLock Lock(&Mutex);
+#endif
 	FMemory::Free(ptr);
 }
 
 static void* VivoxRealloc(void* ptr, size_t bytes)
 {
 	LLM_SCOPE( ELLMTag::AudioVoiceChat );
+#if PLATFORM_IOS
+	FScopeLock Lock(&Mutex);
+#endif
 	return FMemory::Realloc(ptr, bytes);
 }
 
 static void* VivoxCalloc(size_t num, size_t bytes)
 {
 	LLM_SCOPE( ELLMTag::AudioVoiceChat );
+#if PLATFORM_IOS
+	FScopeLock Lock(&Mutex);
+#endif
 	const size_t Size = bytes * num;
 	void* Ret = FMemory::Malloc(Size);
 	FMemory::Memzero(Ret, Size);
@@ -148,12 +164,18 @@ static void* VivoxCalloc(size_t num, size_t bytes)
 static void* VivoxMallocAligned(size_t alignment, size_t bytes)
 {
 	LLM_SCOPE( ELLMTag::AudioVoiceChat );
+#if PLATFORM_IOS
+	FScopeLock Lock(&Mutex);
+#endif
 	return FMemory::Malloc(bytes, alignment);
 }
 
 static void VivoxFreeAligned(void* ptr)
 {
 	LLM_SCOPE( ELLMTag::AudioVoiceChat );
+#if PLATFORM_IOS
+	FScopeLock Lock(&Mutex);
+#endif
 	FMemory::Free(ptr);
 }
 
