[VivoxVoiceChat] Improve logging of some empty device strings

diff --git a/Engine/Plugins/Online/VoiceChat/VivoxVoiceChat/Source/Private/VivoxVoiceChat.cpp b/Engine/Plugins/Online/VoiceChat/VivoxVoiceChat/Source/Private/VivoxVoiceChat.cpp
--- a/Engine/Plugins/Online/VoiceChat/VivoxVoiceChat/Source/Private/VivoxVoiceChat.cpp
+++ b/Engine/Plugins/Online/VoiceChat/VivoxVoiceChat/Source/Private/VivoxVoiceChat.cpp
@@ -1823,18 +1823,18 @@ bool FVivoxVoiceChat::Exec(UWorld* InWorld, const TCHAR* Cmd, FOutputDevice& Ar)
 				const FString DefaultInputDevice = GetDefaultInputDevice();
 				if (InputDevice == DefaultInputDevice)
 				{
-					UE_LOG(LogVivoxVoiceChat, Log, TEXT("    %s (Selected) (Default)"), *DefaultInputDevice);
+					UE_LOG(LogVivoxVoiceChat, Log, TEXT("    [%s] (Selected) (Default)"), *DefaultInputDevice);
 				}
 				else
 				{
-					UE_LOG(LogVivoxVoiceChat, Log, TEXT("    %s (Selected)"), *InputDevice);
-					UE_LOG(LogVivoxVoiceChat, Log, TEXT("    %s (Default)"), *DefaultInputDevice);
+					UE_LOG(LogVivoxVoiceChat, Log, TEXT("    [%s] (Selected)"), *InputDevice);
+					UE_LOG(LogVivoxVoiceChat, Log, TEXT("    [%s] (Default)"), *DefaultInputDevice);
 				}
 				for (const FString& Device : GetAvailableInputDevices())
 				{
 					if (Device != DefaultInputDevice && Device != InputDevice)
 					{
-						UE_LOG(LogVivoxVoiceChat, Log, TEXT("    %s"), *Device);
+						UE_LOG(LogVivoxVoiceChat, Log, TEXT("    [%s]"), *Device);
 					}
 				}
 
@@ -1843,18 +1843,18 @@ bool FVivoxVoiceChat::Exec(UWorld* InWorld, const TCHAR* Cmd, FOutputDevice& Ar)
 				const FString DefaultOutputDevice = GetDefaultOutputDevice();
 				if (OutputDevice == DefaultOutputDevice)
 				{
-					UE_LOG(LogVivoxVoiceChat, Log, TEXT("    %s (Selected) (Default)"), *DefaultOutputDevice);
+					UE_LOG(LogVivoxVoiceChat, Log, TEXT("    [%s] (Selected) (Default)"), *DefaultOutputDevice);
 				}
 				else
 				{
-					UE_LOG(LogVivoxVoiceChat, Log, TEXT("    %s (Selected)"), *OutputDevice);
-					UE_LOG(LogVivoxVoiceChat, Log, TEXT("    %s (Default)"), *DefaultOutputDevice);
+					UE_LOG(LogVivoxVoiceChat, Log, TEXT("    [%s] (Selected)"), *OutputDevice);
+					UE_LOG(LogVivoxVoiceChat, Log, TEXT("    [%s] (Default)"), *DefaultOutputDevice);
 				}
 				for (const FString& Device : GetAvailableOutputDevices())
 				{
 					if (Device != DefaultOutputDevice && Device != OutputDevice)
 					{
-						UE_LOG(LogVivoxVoiceChat, Log, TEXT("    %s"), *Device);
+						UE_LOG(LogVivoxVoiceChat, Log, TEXT("    [%s]"), *Device);
 					}
 				}
 
