[VivoxVoiceChat] Improve logging around block list

diff --git a/Engine/Plugins/Online/VoiceChat/VivoxVoiceChat/Source/Private/VivoxVoiceChat.cpp b/Engine/Plugins/Online/VoiceChat/VivoxVoiceChat/Source/Private/VivoxVoiceChat.cpp
--- a/Engine/Plugins/Online/VoiceChat/VivoxVoiceChat/Source/Private/VivoxVoiceChat.cpp
+++ b/Engine/Plugins/Online/VoiceChat/VivoxVoiceChat/Source/Private/VivoxVoiceChat.cpp
@@ -691,6 +691,13 @@ void FVivoxVoiceChat::BlockPlayers(const TArray<FString>& PlayerNames)
 	{
 		UE_LOG(LogVivoxVoiceChat, Warning, TEXT("BlockPlayers failed: error:%s (%i)"), ANSI_TO_TCHAR(Status.ToString()), Status.GetStatusCode());
 	}
+	else
+	{
+		for (const FString& PlayerName : PlayerNames)
+		{
+			UE_LOG(LogVivoxVoiceChat, Verbose, TEXT("Player blocked: [%s]"), *PlayerName);
+		}
+	}
 }
 
 void FVivoxVoiceChat::UnblockPlayers(const TArray<FString>& PlayerNames)
@@ -706,6 +713,13 @@ void FVivoxVoiceChat::UnblockPlayers(const TArray<FString>& PlayerNames)
 	{
 		UE_LOG(LogVivoxVoiceChat, Warning, TEXT("UnblockPlayers failed: error:%s (%i)"), ANSI_TO_TCHAR(Status.ToString()), Status.GetStatusCode());
 	}
+	else
+	{
+		for (const FString& PlayerName : PlayerNames)
+		{
+			UE_LOG(LogVivoxVoiceChat, Verbose, TEXT("Player unblocked: [%s]"), *PlayerName);
+		}
+	}
 }
 
 void FVivoxVoiceChat::JoinChannel(const FString& ChannelName, const FString& ChannelCredentials, EVoiceChatChannelType ChannelType, const FOnVoiceChatChannelJoinCompleteDelegate& Delegate, TOptional<FVoiceChatChannel3dProperties> Channel3dProperties)
