VivoxVoiceChatPlugin: Add delegates for when connection, login, and channel state changes

diff --git a/Engine/Plugins/Online/VoiceChat/VivoxVoiceChat/Source/Private/VivoxVoiceChat.cpp b/Engine/Plugins/Online/VoiceChat/VivoxVoiceChat/Source/Private/VivoxVoiceChat.cpp
--- a/Engine/Plugins/Online/VoiceChat/VivoxVoiceChat/Source/Private/VivoxVoiceChat.cpp
+++ b/Engine/Plugins/Online/VoiceChat/VivoxVoiceChat/Source/Private/VivoxVoiceChat.cpp
@@ -1209,6 +1209,8 @@ void FVivoxVoiceChat::onLoginCompleted(const VivoxClientApi::AccountName& Accoun
 	LoginSession.State = FLoginSession::EState::LoggedIn;
 
 	TriggerCompletionDelegate(OnVoiceChatLoginCompleteDelegate, PlayerName, ResultSuccess);
+
+	OnVoiceChatLoggedInDelegate.Broadcast(PlayerName);
 }
 
 void FVivoxVoiceChat::onInvalidLoginCredentials(const VivoxClientApi::AccountName& AccountName)
@@ -1243,6 +1245,8 @@ void FVivoxVoiceChat::onLogoutCompleted(const VivoxClientApi::AccountName& Accou
 
 	FString PlayerName = GetPlayerNameFromAccountName(AccountName);
 	TriggerCompletionDelegate(OnVoiceChatLogoutCompleteDelegate, PlayerName, ResultSuccess);
+
+	OnVoiceChatLoggedOutDelegate.Broadcast(PlayerName);
 }
 
 void FVivoxVoiceChat::onLogoutFailed(const VivoxClientApi::AccountName& AccountName, const VivoxClientApi::VCSStatus& Status)
@@ -1262,6 +1266,8 @@ void FVivoxVoiceChat::onChannelJoined(const VivoxClientApi::AccountName& Account
 	FChannelSession& ChannelSession = GetChannelSession(ChannelUri);
 	ChannelSession.State = FChannelSession::EState::Connected;
 	TriggerCompletionDelegate(ChannelSession.JoinDelegate, ChannelSession.ChannelName, ResultSuccess);
+
+	OnVoiceChatChannelJoinedDelegate.Broadcast(ChannelSession.ChannelName);
 }
 
 void FVivoxVoiceChat::onInvalidChannelCredentials(const VivoxClientApi::AccountName& AccountName, const VivoxClientApi::Uri& ChannelUri)
@@ -1303,18 +1309,19 @@ void FVivoxVoiceChat::onChannelExited(const VivoxClientApi::AccountName& Account
 	const bool bWasDisconnecting = ChannelSession.State == FChannelSession::EState::Disconnecting;
 	ChannelSession.State = FChannelSession::EState::Disconnected;
 
-	if (bWasConnected)
-	{
-		OnVoiceChatChannelExitedDelegate.Broadcast(ChannelSession.ChannelName, ResultFromVivoxStatus(ReasonCode));
-	}
-	else if (bWasConnecting)
+	if (bWasConnecting)
 	{
 		// timeouts while connecting call onChannelExited instead of OnChannelJoinFailed
 		TriggerCompletionDelegate(ChannelSession.JoinDelegate, ChannelSession.ChannelName, ResultFromVivoxStatus(ReasonCode));
 	}
-	else if (bWasDisconnecting)
+	else
 	{
-		TriggerCompletionDelegate(ChannelSession.LeaveDelegate, ChannelSession.ChannelName, ResultFromVivoxStatus(ReasonCode));
+		if (bWasDisconnecting)
+		{
+			TriggerCompletionDelegate(ChannelSession.LeaveDelegate, ChannelSession.ChannelName, ResultFromVivoxStatus(ReasonCode));
+		}
+
+		OnVoiceChatChannelExitedDelegate.Broadcast(ChannelSession.ChannelName, ResultFromVivoxStatus(ReasonCode));
 	}
 
 	RemoveChannelSession(ChannelSession.ChannelName);
diff --git a/Engine/Plugins/Online/VoiceChat/VivoxVoiceChat/Source/Public/VivoxVoiceChat.h b/Engine/Plugins/Online/VoiceChat/VivoxVoiceChat/Source/Public/VivoxVoiceChat.h
--- a/Engine/Plugins/Online/VoiceChat/VivoxVoiceChat/Source/Public/VivoxVoiceChat.h
+++ b/Engine/Plugins/Online/VoiceChat/VivoxVoiceChat/Source/Public/VivoxVoiceChat.h
@@ -54,17 +54,21 @@ public:
 	virtual void Disconnect(const FOnVoiceChatDisconnectCompleteDelegate& Delegate) override;
 	virtual bool IsConnecting() const override;
 	virtual bool IsConnected() const override;
+	virtual FOnVoiceChatConnectedDelegate& OnVoiceChatConnected() override { return OnVoiceChatConnectedDelegate; }
 	virtual FOnVoiceChatDisconnectedDelegate& OnVoiceChatDisconnected() override { return OnVoiceChatDisconnectedDelegate; }
 	virtual FOnVoiceChatReconnectedDelegate& OnVoiceChatReconnected() override { return OnVoiceChatReconnectedDelegate; }
 	virtual void Login(FPlatformUserId PlatformId, const FString& PlayerName, const FString& Credentials, const FOnVoiceChatLoginCompleteDelegate& Delegate) override;
 	virtual void Logout(const FOnVoiceChatLogoutCompleteDelegate& Delegate) override;
 	virtual bool IsLoggingIn() const override;
 	virtual bool IsLoggedIn() const override;
+	virtual FOnVoiceChatLoggedInDelegate& OnVoiceChatLoggedIn() override { return OnVoiceChatLoggedInDelegate; }
+	virtual FOnVoiceChatLoggedOutDelegate& OnVoiceChatLoggedOut() override { return OnVoiceChatLoggedOutDelegate; }
 	virtual FString GetLoggedInPlayerName() const override;
 	virtual void BlockPlayers(const TArray<FString>& PlayerNames) override;
 	virtual void UnblockPlayers(const TArray<FString>& PlayerNames) override;
 	virtual void JoinChannel(const FString& ChannelName, const FString& ChannelCredentials, EVoiceChatChannelType ChannelType, const FOnVoiceChatChannelJoinCompleteDelegate& Delegate, TOptional<FVoiceChatChannel3dProperties> Channel3dProperties = TOptional<FVoiceChatChannel3dProperties>()) override;
 	virtual void LeaveChannel(const FString& Channel, const FOnVoiceChatChannelLeaveCompleteDelegate& Delegate) override;
+	virtual FOnVoiceChatChannelJoinedDelegate& OnVoiceChatChannelJoined() override { return OnVoiceChatChannelJoinedDelegate; }
 	virtual FOnVoiceChatChannelExitedDelegate& OnVoiceChatChannelExited() override { return OnVoiceChatChannelExitedDelegate; }
 	virtual FOnVoiceChatCallStatsUpdatedDelegate& OnVoiceChatCallStatsUpdated() override { return OnVoiceChatCallStatsUpdatedDelegate; }
 	virtual void Set3DPosition(const FString& ChannelName, const FVector& SpeakerPosition, const FVector& ListenerPosition, const FVector& ListenerForwardDirection, const FVector& ListenerUpDirection) override;
@@ -218,8 +222,12 @@ protected:
 
 	// Delegates
 	FOnVoiceChatAvailableAudioDevicesChangedDelegate OnVoiceChatAvailableAudioDevicesChangedDelegate;
+	FOnVoiceChatConnectedDelegate OnVoiceChatConnectedDelegate;
 	FOnVoiceChatDisconnectedDelegate OnVoiceChatDisconnectedDelegate;
 	FOnVoiceChatReconnectedDelegate OnVoiceChatReconnectedDelegate;
+	FOnVoiceChatLoggedInDelegate OnVoiceChatLoggedInDelegate;
+	FOnVoiceChatLoggedOutDelegate OnVoiceChatLoggedOutDelegate;
+	FOnVoiceChatChannelJoinedDelegate OnVoiceChatChannelJoinedDelegate;
 	FOnVoiceChatChannelExitedDelegate OnVoiceChatChannelExitedDelegate;
 	FOnVoiceChatPlayerAddedDelegate OnVoiceChatPlayerAddedDelegate;
 	FOnVoiceChatPlayerTalkingUpdatedDelegate OnVoiceChatPlayerTalkingUpdatedDelegate;
diff --git a/Engine/Plugins/Online/VoiceChat/VivoxVoiceChat/Source/Public/VoiceChat.h b/Engine/Plugins/Online/VoiceChat/VivoxVoiceChat/Source/Public/VoiceChat.h
--- a/Engine/Plugins/Online/VoiceChat/VivoxVoiceChat/Source/Public/VoiceChat.h
+++ b/Engine/Plugins/Online/VoiceChat/VivoxVoiceChat/Source/Public/VoiceChat.h
@@ -85,8 +85,12 @@ DECLARE_DELEGATE_TwoParams(FOnVoiceChatChannelLeaveCompleteDelegate, const FStri
 
 DECLARE_MULTICAST_DELEGATE(FOnVoiceChatAvailableAudioDevicesChangedDelegate);
 DECLARE_MULTICAST_DELEGATE(FOnVoiceChatReconnectedDelegate);
+DECLARE_MULTICAST_DELEGATE(FOnVoiceChatConnectedDelegate);
 DECLARE_MULTICAST_DELEGATE_OneParam(FOnVoiceChatDisconnectedDelegate, const FVoiceChatResult& /* Reason */);
+DECLARE_MULTICAST_DELEGATE_OneParam(FOnVoiceChatLoggedInDelegate, const FString& /* PlayerName */);
+DECLARE_MULTICAST_DELEGATE_OneParam(FOnVoiceChatLoggedOutDelegate, const FString& /* PlayerName */);
 DECLARE_MULTICAST_DELEGATE_OneParam(FOnVoiceChatCallStatsUpdatedDelegate, const FVoiceChatCallStats& /* CallStats */);
+DECLARE_MULTICAST_DELEGATE_OneParam(FOnVoiceChatChannelJoinedDelegate, const FString& /* ChannelName */);
 DECLARE_MULTICAST_DELEGATE_TwoParams(FOnVoiceChatChannelExitedDelegate, const FString& /* ChannelName */, const FVoiceChatResult& /* Reason */);
 DECLARE_MULTICAST_DELEGATE_TwoParams(FOnVoiceChatPlayerAddedDelegate, const FString& /* ChannelName */, const FString& /* PlayerName */);
 DECLARE_MULTICAST_DELEGATE_TwoParams(FOnVoiceChatPlayerRemovedDelegate, const FString& /* ChannelName */, const FString& /* PlayerName */);
@@ -269,7 +273,12 @@ public:
 	virtual bool IsConnected() const = 0;
 
 	/**
-	 * Delegate triggered when we are unexpectedly disconnected from voice chat
+	 * Delegate triggered when we are connected to voice chat
+	 */
+	virtual FOnVoiceChatConnectedDelegate& OnVoiceChatConnected() = 0;
+
+	/**
+	 * Delegate triggered when we are disconnected from voice chat
 	 */
 	virtual FOnVoiceChatDisconnectedDelegate& OnVoiceChatDisconnected() = 0;
 
@@ -311,6 +320,16 @@ public:
 	 */
 	virtual bool IsLoggedIn() const = 0;
 
+	/**
+	 * Delegate triggered when we are logged in to voice chat
+	 */
+	virtual FOnVoiceChatLoggedInDelegate& OnVoiceChatLoggedIn() = 0;
+
+	/**
+	 * Delegate triggered when we are logged out from voice chat
+	 */
+	virtual FOnVoiceChatLoggedOutDelegate& OnVoiceChatLoggedOut() = 0;
+
 	/**
 	 * Get the player name used to log in
 	 *
@@ -352,7 +371,12 @@ public:
 	virtual void LeaveChannel(const FString& ChannelName, const FOnVoiceChatChannelLeaveCompleteDelegate& Delegate) = 0;
 
 	/**
-	 * Delegate triggered when we unexpectedly leave a voice channel
+	 * Delegate triggered when we join a voice channel
+	 */
+	virtual FOnVoiceChatChannelJoinedDelegate& OnVoiceChatChannelJoined() = 0;
+
+	/**
+	 * Delegate triggered when we leave a voice channel
 	 */
 	virtual FOnVoiceChatChannelExitedDelegate& OnVoiceChatChannelExited() = 0;
 
