Wake the game thread when vivox events are received

diff --git a/Engine/Plugins/Online/VoiceChat/VivoxVoiceChat/Source/Private/VivoxVoiceChat.cpp b/Engine/Plugins/Online/VoiceChat/VivoxVoiceChat/Source/Private/VivoxVoiceChat.cpp
--- a/Engine/Plugins/Online/VoiceChat/VivoxVoiceChat/Source/Private/VivoxVoiceChat.cpp
+++ b/Engine/Plugins/Online/VoiceChat/VivoxVoiceChat/Source/Private/VivoxVoiceChat.cpp
@@ -5,6 +5,7 @@
 #include "Async/Async.h"
 #include "Logging/LogMacros.h"
 #include "Misc/ConfigCacheIni.h"
+#include "Misc/EmbeddedCommunication.h"
 #include "Misc/ScopeLock.h"
 #include "Modules/ModuleManager.h"
 #include "HAL/LowLevelMemTracker.h"
@@ -1131,6 +1132,8 @@ void FVivoxVoiceChat::InvokeOnUIThread(void (Func)(void* Arg0), void* Arg0)
 			(*Func)(Arg0);
 		}
 	});
+
+	FEmbeddedCommunication::WakeGameThread();
 }
 
 void FVivoxVoiceChat::onLogStatementEmitted(LogLevel Level, long long NativeMillisecondsSinceEpoch, long ThreadId, const char* LogMessage)
