VivoxVoiceChat Plugin - Respect EAudioFeature::BackgroundAudio setting in IOSAppDelegate

#IOS

diff --git a/Engine/Plugins/Online/VoiceChat/VivoxVoiceChat/Source/Private/IOS/IOSVivoxVoiceChat.cpp b/Engine/Plugins/Online/VoiceChat/VivoxVoiceChat/Source/Private/IOS/IOSVivoxVoiceChat.cpp
--- a/Engine/Plugins/Online/VoiceChat/VivoxVoiceChat/Source/Private/IOS/IOSVivoxVoiceChat.cpp
+++ b/Engine/Plugins/Online/VoiceChat/VivoxVoiceChat/Source/Private/IOS/IOSVivoxVoiceChat.cpp
@@ -131,7 +131,8 @@ void FIOSVivoxVoiceChat::HandleApplicationWillEnterBackground()
 {
 	UE_LOG(LogVivoxVoiceChat, Log, TEXT("OnApplicationWillEnterBackgroundDelegate"));
 
-	if (IsConnected() && bDisconnectInBackground)
+	const bool bIsBackgroundAudioEnabled = [[IOSAppDelegate GetDelegate] IsFeatureActive:EAudioFeature::BackgroundAudio];
+	if (IsConnected() && bDisconnectInBackground && !bIsBackgroundAudioEnabled)
 	{
 		if (BGTask != UIBackgroundTaskInvalid)
 		{
