VivoxVoiceChat plugin: Fix FVivoxVoiceChat::GetAudioOutputVolume returning a value in [40,60] instead of [0.0,1.0]

diff --git a/Engine/Plugins/Online/VoiceChat/VivoxVoiceChat/Source/Private/VivoxVoiceChat.cpp b/Engine/Plugins/Online/VoiceChat/VivoxVoiceChat/Source/Private/VivoxVoiceChat.cpp
--- a/Engine/Plugins/Online/VoiceChat/VivoxVoiceChat/Source/Private/VivoxVoiceChat.cpp
+++ b/Engine/Plugins/Online/VoiceChat/VivoxVoiceChat/Source/Private/VivoxVoiceChat.cpp
@@ -307,12 +307,12 @@ void FVivoxVoiceChat::SetAudioOutputVolume(float InVolume)
 
 float FVivoxVoiceChat::GetAudioInputVolume() const
 {
-	return VivoxClientConnection.GetMasterAudioInputDeviceVolume();
+	return static_cast<float>(VivoxClientConnection.GetMasterAudioInputDeviceVolume() - VIVOX_MIN_VOL) / static_cast<float>(VIVOX_MAX_VOL - VIVOX_MIN_VOL);
 }
 
 float FVivoxVoiceChat::GetAudioOutputVolume() const
 {
-	return VivoxClientConnection.GetMasterAudioOutputDeviceVolume();
+	return static_cast<float>(VivoxClientConnection.GetMasterAudioOutputDeviceVolume() - VIVOX_MIN_VOL) / static_cast<float>(VIVOX_MAX_VOL - VIVOX_MIN_VOL);
 }
 
 void FVivoxVoiceChat::SetAudioInputDeviceMuted(bool bIsMuted)
@@ -1842,7 +1842,7 @@ bool FVivoxVoiceChat::Exec(UWorld* InWorld, const TCHAR* Cmd, FOutputDevice& Ar)
 			UE_LOG(LogVivoxVoiceChat, Log, TEXT("Initialized: %s"), *LexToString(IsInitialized()));
 			if (IsInitialized())
 			{
-				UE_LOG(LogVivoxVoiceChat, Log, TEXT("  Input Devices: muted:%s volume:%.2f"), *LexToString(GetAudioInputDeviceMuted()), static_cast<float>(GetAudioInputVolume() - VIVOX_MIN_VOL) / static_cast<float>(VIVOX_MAX_VOL - VIVOX_MIN_VOL));
+				UE_LOG(LogVivoxVoiceChat, Log, TEXT("  Input Devices: muted:%s volume:%.2f"), *LexToString(GetAudioInputDeviceMuted()), GetAudioInputVolume());
 				const FString InputDevice = GetInputDevice();
 				const FString DefaultInputDevice = GetDefaultInputDevice();
 				if (InputDevice == DefaultInputDevice)
@@ -1862,7 +1862,7 @@ bool FVivoxVoiceChat::Exec(UWorld* InWorld, const TCHAR* Cmd, FOutputDevice& Ar)
 					}
 				}
 
-				UE_LOG(LogVivoxVoiceChat, Log, TEXT("  Output Devices: muted:%s volume:%.2f"), *LexToString(GetAudioOutputDeviceMuted()), static_cast<float>(GetAudioOutputVolume() - VIVOX_MIN_VOL) / static_cast<float>(VIVOX_MAX_VOL - VIVOX_MIN_VOL));
+				UE_LOG(LogVivoxVoiceChat, Log, TEXT("  Output Devices: muted:%s volume:%.2f"), *LexToString(GetAudioOutputDeviceMuted()), GetAudioOutputVolume());
 				const FString OutputDevice = GetOutputDevice();
 				const FString DefaultOutputDevice = GetDefaultOutputDevice();
 				if (OutputDevice == DefaultOutputDevice)
