[IOS] VivoxVoiceChat plugin: Provide a way to disconnect from voice chat after being backgrounded for a specified period of time

diff --git a/Engine/Plugins/Online/VoiceChat/VivoxVoiceChat/Source/Private/IOS/IOSVivoxVoiceChat.cpp b/Engine/Plugins/Online/VoiceChat/VivoxVoiceChat/Source/Private/IOS/IOSVivoxVoiceChat.cpp
--- a/Engine/Plugins/Online/VoiceChat/VivoxVoiceChat/Source/Private/IOS/IOSVivoxVoiceChat.cpp
+++ b/Engine/Plugins/Online/VoiceChat/VivoxVoiceChat/Source/Private/IOS/IOSVivoxVoiceChat.cpp
@@ -17,6 +17,8 @@ FIOSVivoxVoiceChat::FIOSVivoxVoiceChat()
 	, bInBackground(false)
 	, bShouldReconnect(false)
 	, bIsRecording(false)
+	, BackgroundDelayedDisconnectTime(0.0f)
+	, DelayedDisconnectTimer(nullptr)
 {
 }
 
@@ -31,6 +33,7 @@ bool FIOSVivoxVoiceChat::Initialize()
 	if (bResult)
 	{
 		GConfig->GetBool(TEXT("VoiceChat.Vivox"), TEXT("bDisconnectInBackground"), bDisconnectInBackground, GEngineIni);
+		GConfig->GetFloat(TEXT("VoiceChat.Vivox"), TEXT("BackgroundDelayedDisconnectTime"), BackgroundDelayedDisconnectTime, GEngineIni);
 		
 		if (!ApplicationWillEnterBackgroundHandle.IsValid())
 		{
@@ -129,6 +132,11 @@ void FIOSVivoxVoiceChat::OnVoiceChatDisconnectComplete(const FVoiceChatResult& R
 	}
 }
 
+void FIOSVivoxVoiceChat::OnVoiceChatDelayedDisconnectComplete(const FVoiceChatResult& Result)
+{
+	OnVoiceChatDisconnectedDelegate.Broadcast(Result);
+}
+
 void FIOSVivoxVoiceChat::SetVivoxSdkConfigHints(vx_sdk_config_t &Hints)
 {
 	Hints.dynamic_voice_processing_switching = 0;
@@ -159,9 +167,23 @@ void FIOSVivoxVoiceChat::HandleApplicationWillEnterBackground()
 	}
 	else
 	{
+		if (BackgroundDelayedDisconnectTime > 0.01)
+		{
+			dispatch_async(dispatch_get_main_queue(), ^
+			{
+				DelayedDisconnectTimer = [NSTimer scheduledTimerWithTimeInterval:BackgroundDelayedDisconnectTime repeats:NO block:^(NSTimer * _Nonnull timer) {
+					[FIOSAsyncTask CreateTaskWithBlock:^bool(void){
+						Disconnect(FOnVoiceChatDisconnectCompleteDelegate::CreateRaw(this, &FIOSVivoxVoiceChat::OnVoiceChatDelayedDisconnectComplete));
+						return true;
+					}];
+					DelayedDisconnectTimer = nullptr;
+				}];
+			});
+			
+		}
 		bShouldReconnect = false;
 	}
-
+	
 	VivoxClientConnection.EnteredBackground();
 }
 
@@ -177,6 +199,13 @@ void FIOSVivoxVoiceChat::HandleApplicationHasEnteredForeground()
 		[App endBackgroundTask : BGTask];
 		BGTask = UIBackgroundTaskInvalid;
 	}
+
+	dispatch_async(dispatch_get_main_queue(), ^
+	{
+		[DelayedDisconnectTimer invalidate];
+		DelayedDisconnectTimer = nullptr;
+	});
+
 	if (bShouldReconnect)
 	{
 		Reconnect();
diff --git a/Engine/Plugins/Online/VoiceChat/VivoxVoiceChat/Source/Private/IOS/IOSVivoxVoiceChat.h b/Engine/Plugins/Online/VoiceChat/VivoxVoiceChat/Source/Private/IOS/IOSVivoxVoiceChat.h
--- a/Engine/Plugins/Online/VoiceChat/VivoxVoiceChat/Source/Private/IOS/IOSVivoxVoiceChat.h
+++ b/Engine/Plugins/Online/VoiceChat/VivoxVoiceChat/Source/Private/IOS/IOSVivoxVoiceChat.h
@@ -28,6 +28,7 @@ protected:
 private:
 	void OnVoiceChatConnectComplete(const FVoiceChatResult& Result);
 	void OnVoiceChatDisconnectComplete(const FVoiceChatResult& Result);
+	void OnVoiceChatDelayedDisconnectComplete(const FVoiceChatResult& Result);
 
 	void HandleApplicationWillEnterBackground();
 	void HandleApplicationHasEnteredForeground();
@@ -43,4 +44,7 @@ private:
 	bool bShouldReconnect;
 
 	bool bIsRecording;
+
+	float BackgroundDelayedDisconnectTime;
+	NSTimer* DelayedDisconnectTimer;
 };
