Fix Mac editor compile error

diff --git a/Engine/Plugins/Online/VoiceChat/VivoxVoiceChat/Source/Public/VivoxVoiceChat.h b/Engine/Plugins/Online/VoiceChat/VivoxVoiceChat/Source/Public/VivoxVoiceChat.h
--- a/Engine/Plugins/Online/VoiceChat/VivoxVoiceChat/Source/Public/VivoxVoiceChat.h
+++ b/Engine/Plugins/Online/VoiceChat/VivoxVoiceChat/Source/Public/VivoxVoiceChat.h
@@ -23,7 +23,7 @@ DECLARE_LOG_CATEGORY_EXTERN(LogVivoxVoiceChat, Log, All);
 
 DECLARE_STATS_GROUP(TEXT("Vivox"), STATGROUP_Vivox, STATCAT_Advanced);
 
-VIVOXVOICECHAT_API class FVivoxDelegates
+class VIVOXVOICECHAT_API FVivoxDelegates
 {
 public:
 	/** Delegate called when the status of the audio device has changed. Triggered on platforms that the vivox sdk provides this event for */
