[VivoxVoiceChat] Avoid spamming the log when vivox sends repeated log messages

diff --git a/Engine/Plugins/Online/VoiceChat/VivoxVoiceChat/Source/Private/VivoxVoiceChat.cpp b/Engine/Plugins/Online/VoiceChat/VivoxVoiceChat/Source/Private/VivoxVoiceChat.cpp
--- a/Engine/Plugins/Online/VoiceChat/VivoxVoiceChat/Source/Private/VivoxVoiceChat.cpp
+++ b/Engine/Plugins/Online/VoiceChat/VivoxVoiceChat/Source/Private/VivoxVoiceChat.cpp
@@ -1138,26 +1138,65 @@ void FVivoxVoiceChat::InvokeOnUIThread(void (Func)(void* Arg0), void* Arg0)
 	FEmbeddedCommunication::WakeGameThread();
 }
 
-void FVivoxVoiceChat::onLogStatementEmitted(LogLevel Level, long long NativeMillisecondsSinceEpoch, long ThreadId, const char* LogMessage)
+FString GetLogLevelString(VivoxClientApi::IClientApiEventHandler::LogLevel Level)
 {
-	if (Level == LogLevelError)
+	FString LogLevelString;
+	switch (Level)
 	{
-		UE_LOG(LogVivoxVoiceChat, Warning, TEXT("vivox: Error: %s"), ANSI_TO_TCHAR(LogMessage));
+	case VivoxClientApi::IClientApiEventHandler::LogLevel::LogLevelError:		LogLevelString = TEXT("Error"); break;
+	case VivoxClientApi::IClientApiEventHandler::LogLevel::LogLevelWarning:		LogLevelString = TEXT("Warning"); break;
+	case VivoxClientApi::IClientApiEventHandler::LogLevel::LogLevelInfo:		LogLevelString = TEXT("Info"); break;
+	case VivoxClientApi::IClientApiEventHandler::LogLevel::LogLevelDebug:		LogLevelString = TEXT("Debug"); break;
+	case VivoxClientApi::IClientApiEventHandler::LogLevel::LogLevelTrace:		LogLevelString = TEXT("Trace"); break;
+	default:																	LogLevelString = TEXT("Unknown"); break;
 	}
-	else
+	return LogLevelString;
+}
+
+void FVivoxVoiceChat::onLogStatementEmitted(LogLevel Level, long long NativeMillisecondsSinceEpoch, long ThreadId, const char* LogMessageCStr)
+{
+	FString LogMessage = ANSI_TO_TCHAR(LogMessageCStr);
+	bool bLogMatches = false;
+
+	// Don't spam the log if we receive repeated log messages
+	if (!LastLogMessage.IsEmpty())
+	{
+		bLogMatches = LogMessage == LastLogMessage && Level == LastLogLevel;
+		if (bLogMatches)
+		{
+			LogSpamCount++;
+		}
+		
+		const bool bLogSpamCountPowerOfTwo = FMath::IsPowerOfTwo(LogSpamCount);
+		const bool bLogDueToPot = bLogMatches && bLogSpamCountPowerOfTwo;
+		const bool bLogDueToChange = !bLogMatches && LogSpamCount > 0 && !bLogSpamCountPowerOfTwo; // Don't log on change if POT, as we already logged it.
+		if(bLogDueToChange || bLogDueToPot)
+		{
+			if (LastLogLevel == LogLevelError)
+			{
+				UE_LOG(LogVivoxVoiceChat, Warning, TEXT("vivox: Error: %s (seen %d times)"), *LastLogMessage, LogSpamCount);
+			}
+			else
+			{
+				UE_LOG(LogVivoxVoiceChat, Log, TEXT("vivox: %s: %s (seen %d times)"), *GetLogLevelString(LastLogLevel), *LastLogMessage, LogSpamCount);
+			}
+		}
+	}
+
+	if (!bLogMatches)
 	{
-		FString LogLevelString;
-		switch (Level)
+		if (Level == LogLevelError)
+		{
+			UE_LOG(LogVivoxVoiceChat, Log, TEXT("vivox: Error: %s"), *LogMessage);
+		}
+		else
 		{
-		case LogLevelError:		LogLevelString = TEXT("Error"); break;
-		case LogLevelWarning:	LogLevelString = TEXT("Warning"); break;
-		case LogLevelInfo:		LogLevelString = TEXT("Info"); break;
-		case LogLevelDebug:		LogLevelString = TEXT("Debug"); break;
-		case LogLevelTrace:		LogLevelString = TEXT("Trace"); break;
-		default:				LogLevelString = TEXT("Unknown"); break;
+			UE_LOG(LogVivoxVoiceChat, Log, TEXT("vivox: %s: %s"), *GetLogLevelString(Level), *LogMessage);
 		}
 
-		UE_LOG(LogVivoxVoiceChat, Log, TEXT("vivox: %s: %s"), *LogLevelString, ANSI_TO_TCHAR(LogMessage));
+		LogSpamCount = 0;
+		LastLogMessage = MoveTemp(LogMessage);
+		LastLogLevel = Level;
 	}
 }
 
diff --git a/Engine/Plugins/Online/VoiceChat/VivoxVoiceChat/Source/Public/VivoxVoiceChat.h b/Engine/Plugins/Online/VoiceChat/VivoxVoiceChat/Source/Public/VivoxVoiceChat.h
--- a/Engine/Plugins/Online/VoiceChat/VivoxVoiceChat/Source/Public/VivoxVoiceChat.h
+++ b/Engine/Plugins/Online/VoiceChat/VivoxVoiceChat/Source/Public/VivoxVoiceChat.h
@@ -281,6 +281,11 @@ protected:
 
 	void ClearLoginSession();
 
+	// Log spam avoidance
+	FString LastLogMessage;
+	LogLevel LastLogLevel;
+	int LogSpamCount = 0;
+
 	// ~Begin FSelfRegisteringExec Interface
 	virtual bool Exec(UWorld* InWorld, const TCHAR* Cmd, FOutputDevice& Ar) override;
 	// ~End FSelfRegisteringExec Interface
