VivoxVoiceChat plugin: Add an interface for setting implementation specific settings. Use this to allow iOS HardwareAEC to be enabled/disabled at runtime.

diff --git a/Engine/Plugins/Online/VoiceChat/VivoxVoiceChat/Source/Private/IOS/IOSVivoxVoiceChat.cpp b/Engine/Plugins/Online/VoiceChat/VivoxVoiceChat/Source/Private/IOS/IOSVivoxVoiceChat.cpp
--- a/Engine/Plugins/Online/VoiceChat/VivoxVoiceChat/Source/Private/IOS/IOSVivoxVoiceChat.cpp
+++ b/Engine/Plugins/Online/VoiceChat/VivoxVoiceChat/Source/Private/IOS/IOSVivoxVoiceChat.cpp
@@ -34,6 +34,7 @@ bool FIOSVivoxVoiceChat::Initialize()
 	{
 		GConfig->GetBool(TEXT("VoiceChat.Vivox"), TEXT("bDisconnectInBackground"), bDisconnectInBackground, GEngineIni);
 		GConfig->GetFloat(TEXT("VoiceChat.Vivox"), TEXT("BackgroundDelayedDisconnectTime"), BackgroundDelayedDisconnectTime, GEngineIni);
+		GConfig->GetBool(TEXT("VoiceChat.Vivox"), TEXT("bEnableHardwareAEC"), bEnableHardwareAEC, GEngineIni);
 		
 		if (!ApplicationWillEnterBackgroundHandle.IsValid())
 		{
@@ -45,8 +46,6 @@ bool FIOSVivoxVoiceChat::Initialize()
 		}
 	}
 
-	vx_set_platform_aec_enabled(1);
-
 	bInBackground = false;
 	bShouldReconnect = false;
 	bIsRecording = false;
@@ -70,9 +69,31 @@ bool FIOSVivoxVoiceChat::Uninitialize()
 	return FVivoxVoiceChat::Uninitialize();
 }
 
+void FIOSVivoxVoiceChat::SetSetting(const FString& Name, const FString& Value)
+{
+	if (Name == TEXT("HardwareAEC"))
+	{
+		OverrideEnableHardwareAEC = FCString::ToBool(*Value);
+		UpdateVoiceChatSettings();
+	}
+	else
+	{
+		FVivoxVoiceChat::SetSetting(Name, Value);
+	}
+}
+
+FString FIOSVivoxVoiceChat::GetSetting(const FString& Name)
+{
+	if (Name == TEXT("HardwareAEC"))
+	{
+		return LexToString(OverrideEnableHardwareAEC.Get(bEnableHardwareAEC));
+	}
+	return FVivoxVoiceChat::GetSetting(Name);
+}
+
 FDelegateHandle FIOSVivoxVoiceChat::StartRecording(const FOnVoiceChatRecordSamplesAvailableDelegate::FDelegate& Delegate)
 {
-	FPlatformMisc::EnableVoiceChat(true);
+	EnableVoiceChat(true);
 	return FVivoxVoiceChat::StartRecording(Delegate);
 }
 
@@ -81,13 +102,13 @@ void FIOSVivoxVoiceChat::StopRecording(FDelegateHandle Handle)
 	FVivoxVoiceChat::StopRecording(Handle);
 	if (ConnectionState < EConnectionState::Connecting)
 	{
-		FPlatformMisc::EnableVoiceChat(false);
+		EnableVoiceChat(false);
 	}
 }
 
 void FIOSVivoxVoiceChat::onConnectCompleted(const VivoxClientApi::Uri& Server)
 {
-	FPlatformMisc::EnableVoiceChat(true);
+	EnableVoiceChat(true);
 	FVivoxVoiceChat::onConnectCompleted(Server);
 }
 
@@ -96,7 +117,7 @@ void FIOSVivoxVoiceChat::onDisconnected(const VivoxClientApi::Uri& Server, const
 	FVivoxVoiceChat::onDisconnected(Server, Status);
 	if (!bIsRecording)
 	{
-		FPlatformMisc::EnableVoiceChat(false);
+		EnableVoiceChat(false);
 	}
 }
 
@@ -217,3 +238,35 @@ void FIOSVivoxVoiceChat::Reconnect()
 	Connect(FOnVoiceChatConnectCompleteDelegate::CreateRaw(this, &FIOSVivoxVoiceChat::OnVoiceChatConnectComplete));
 	bShouldReconnect = false;
 }
+
+bool FIOSVivoxVoiceChat::IsHardwareAECEnabled() const
+{
+	return OverrideEnableHardwareAEC.Get(bEnableHardwareAEC);
+}
+
+void FIOSVivoxVoiceChat::EnableVoiceChat(bool bEnable)
+{
+	if (FPlatformMisc::IsVoiceChatEnabled() != bEnable)
+	{
+		if (IsHardwareAECEnabled())
+		{
+			[[IOSAppDelegate GetDelegate] EnableHighQualityVoiceChat:bEnable];
+			vx_set_platform_aec_enabled(bEnable ? 1 : 0);
+		}
+		FPlatformMisc::EnableVoiceChat(bEnable);
+	}
+}
+
+void FIOSVivoxVoiceChat::UpdateVoiceChatSettings()
+{
+	if (FPlatformMisc::IsVoiceChatEnabled())
+	{
+		// update the aec settings
+		const bool bEnableAEC = IsHardwareAECEnabled();
+		[[IOSAppDelegate GetDelegate] EnableHighQualityVoiceChat:bEnableAEC];
+		vx_set_platform_aec_enabled(bEnableAEC ? 1 : 0);
+		// reenable voice chat to apply the changes
+		FPlatformMisc::EnableVoiceChat(true);
+	}
+}
+
diff --git a/Engine/Plugins/Online/VoiceChat/VivoxVoiceChat/Source/Private/IOS/IOSVivoxVoiceChat.h b/Engine/Plugins/Online/VoiceChat/VivoxVoiceChat/Source/Private/IOS/IOSVivoxVoiceChat.h
--- a/Engine/Plugins/Online/VoiceChat/VivoxVoiceChat/Source/Private/IOS/IOSVivoxVoiceChat.h
+++ b/Engine/Plugins/Online/VoiceChat/VivoxVoiceChat/Source/Private/IOS/IOSVivoxVoiceChat.h
@@ -13,6 +13,8 @@ public:
 	// ~Begin IVoiceChat Interface
 	virtual bool Initialize() override;
 	virtual bool Uninitialize() override;
+	virtual void SetSetting(const FString& Name, const FString& Value) override;
+	virtual FString GetSetting(const FString& Name) override;
 	virtual FDelegateHandle StartRecording(const FOnVoiceChatRecordSamplesAvailableDelegate::FDelegate& Delegate) override;
 	virtual void StopRecording(FDelegateHandle Handle) override;
 	// ~End IVoiceChat Interface
@@ -47,4 +49,12 @@ private:
 
 	float BackgroundDelayedDisconnectTime;
 	NSTimer* DelayedDisconnectTimer;
+
+	bool IsHardwareAECEnabled() const;
+	bool bEnableHardwareAEC = false;
+	TOptional<bool> OverrideEnableHardwareAEC;
+
+	void EnableVoiceChat(bool bEnable);
+	void UpdateVoiceChatSettings();
+	void DisableVoiceChat();
 };
diff --git a/Engine/Plugins/Online/VoiceChat/VivoxVoiceChat/Source/Private/VivoxVoiceChat.cpp b/Engine/Plugins/Online/VoiceChat/VivoxVoiceChat/Source/Private/VivoxVoiceChat.cpp
--- a/Engine/Plugins/Online/VoiceChat/VivoxVoiceChat/Source/Private/VivoxVoiceChat.cpp
+++ b/Engine/Plugins/Online/VoiceChat/VivoxVoiceChat/Source/Private/VivoxVoiceChat.cpp
@@ -280,6 +280,8 @@ bool FVivoxVoiceChat::Uninitialize()
 	if (IsInitialized())
 	{
 		VivoxClientConnection.Uninitialize();
+		ConnectionState = EConnectionState::Disconnected;
+		LoginSession = FLoginSession();
 		bInitialized = false;
 	}
 
@@ -291,6 +293,15 @@ bool FVivoxVoiceChat::IsInitialized() const
 	return bInitialized;
 }
 
+void FVivoxVoiceChat::SetSetting(const FString& Name, const FString& Value)
+{
+}
+
+FString FVivoxVoiceChat::GetSetting(const FString& Name)
+{
+	return FString();
+}
+
 void FVivoxVoiceChat::SetAudioInputVolume(float InVolume)
 {
 	UE_LOG(LogVivoxVoiceChat, Verbose, TEXT("SetAudioInputVolume %f"), InVolume);
diff --git a/Engine/Plugins/Online/VoiceChat/VivoxVoiceChat/Source/Public/VivoxVoiceChat.h b/Engine/Plugins/Online/VoiceChat/VivoxVoiceChat/Source/Public/VivoxVoiceChat.h
--- a/Engine/Plugins/Online/VoiceChat/VivoxVoiceChat/Source/Public/VivoxVoiceChat.h
+++ b/Engine/Plugins/Online/VoiceChat/VivoxVoiceChat/Source/Public/VivoxVoiceChat.h
@@ -41,6 +41,8 @@ public:
 	virtual bool Initialize() override;
 	virtual bool Uninitialize() override;
 	virtual bool IsInitialized() const override;
+	virtual void SetSetting(const FString& Name, const FString& Value) override;
+	virtual FString GetSetting(const FString& Name) override;
 	virtual void SetAudioInputVolume(float Volume) override;
 	virtual void SetAudioOutputVolume(float Volume) override;
 	virtual float GetAudioInputVolume() const override;
diff --git a/Engine/Plugins/Online/VoiceChat/VivoxVoiceChat/Source/Public/VoiceChat.h b/Engine/Plugins/Online/VoiceChat/VivoxVoiceChat/Source/Public/VoiceChat.h
--- a/Engine/Plugins/Online/VoiceChat/VivoxVoiceChat/Source/Public/VoiceChat.h
+++ b/Engine/Plugins/Online/VoiceChat/VivoxVoiceChat/Source/Public/VoiceChat.h
@@ -135,6 +135,23 @@ public:
 	 */
 	virtual bool IsInitialized() const = 0;
 
+	/**
+	 * Set an implementation specific setting
+	 *
+	 * @param Name Name of the config setting to set
+	 * @param Value Value for the setting
+	 */
+	virtual void SetSetting(const FString& Name, const FString& Value) = 0;
+
+	/**
+	 * Get an implementation specific setting
+	 *
+	 * @param Name Name of the config setting to get
+	 *
+	 * @return Current value for the setting
+	 */
+	virtual FString GetSetting(const FString& Name) = 0;
+
 	/**
 	 * Set record volume
 	 *
