Remove return that shouldn't be there

diff --git a/Engine/Plugins/Online/VoiceChat/VivoxVoiceChat/Source/Private/VivoxVoiceChat.cpp b/Engine/Plugins/Online/VoiceChat/VivoxVoiceChat/Source/Private/VivoxVoiceChat.cpp
--- a/Engine/Plugins/Online/VoiceChat/VivoxVoiceChat/Source/Private/VivoxVoiceChat.cpp
+++ b/Engine/Plugins/Online/VoiceChat/VivoxVoiceChat/Source/Private/VivoxVoiceChat.cpp
@@ -1038,7 +1038,7 @@ void FVivoxVoiceChat::UnregisterOnVoiceChatBeforeRecvAudioRenderedDelegate(FDele
 {
 	FScopeLock Lock(&BeforeRecvAudioRenderedLock);
 
-	return OnVoiceChatBeforeRecvAudioRenderedDelegate.Remove(Handle);
+	OnVoiceChatBeforeRecvAudioRenderedDelegate.Remove(Handle);
 }
 
 FString FVivoxVoiceChat::InsecureGetLoginToken(const FString& PlayerName)
