Fix Vivox Disconnect completion delegate not firing

diff --git a/Engine/Plugins/Online/VoiceChat/VivoxVoiceChat/Source/Private/VivoxVoiceChat.cpp b/Engine/Plugins/Online/VoiceChat/VivoxVoiceChat/Source/Private/VivoxVoiceChat.cpp
--- a/Engine/Plugins/Online/VoiceChat/VivoxVoiceChat/Source/Private/VivoxVoiceChat.cpp
+++ b/Engine/Plugins/Online/VoiceChat/VivoxVoiceChat/Source/Private/VivoxVoiceChat.cpp
@@ -499,6 +499,7 @@ void FVivoxVoiceChat::Disconnect(const FOnVoiceChatDisconnectCompleteDelegate& D
 		VivoxClientApi::Uri BackendUri(TCHAR_TO_ANSI(*VivoxServerUrl));
 		VivoxClientConnection.Disconnect(BackendUri);
 
+		ConnectionState = EConnectionState::Disconnecting;
 		OnVoiceChatDisconnectCompleteDelegates.Add(Delegate);
 	}
 	else
