Fix android/ios crash when backgrounding after uninitializing the VivoxVoiceChat plugin

diff --git a/Engine/Plugins/Online/VoiceChat/VivoxVoiceChat/Source/Private/Android/AndroidVivoxVoiceChat.cpp b/Engine/Plugins/Online/VoiceChat/VivoxVoiceChat/Source/Private/Android/AndroidVivoxVoiceChat.cpp
--- a/Engine/Plugins/Online/VoiceChat/VivoxVoiceChat/Source/Private/Android/AndroidVivoxVoiceChat.cpp
+++ b/Engine/Plugins/Online/VoiceChat/VivoxVoiceChat/Source/Private/Android/AndroidVivoxVoiceChat.cpp
@@ -82,6 +82,22 @@ bool FAndroidVivoxVoiceChat::Initialize()
 	return bReturn;
 }
 
+bool FAndroidVivoxVoiceChat::Uninitialize()
+{
+	if (ApplicationWillEnterBackgroundHandle.IsValid())
+	{
+		FCoreDelegates::ApplicationWillDeactivateDelegate.Remove(ApplicationWillEnterBackgroundHandle);
+		ApplicationWillEnterBackgroundHandle.Reset();
+	}
+	if (ApplicationDidEnterForegroundHandle.IsValid())
+	{
+		FCoreDelegates::ApplicationHasReactivatedDelegate.Remove(ApplicationDidEnterForegroundHandle);
+		ApplicationDidEnterForegroundHandle.Reset();
+	}
+
+	return FVivoxVoiceChat::Uninitialize();
+}
+
 void FAndroidVivoxVoiceChat::OnVoiceChatConnectComplete(const FVoiceChatResult& Result)
 {
 	if (Result.bSuccess)
@@ -100,7 +116,7 @@ void FAndroidVivoxVoiceChat::OnVoiceChatDisconnectComplete(const FVoiceChatResul
 	{
 		bShouldReconnect = true;
 	}
-	else
+	else if (IsInitialized())
 	{
 		// disconnect complete delegate fired after entering foreground
 		Reconnect();
diff --git a/Engine/Plugins/Online/VoiceChat/VivoxVoiceChat/Source/Private/Android/AndroidVivoxVoiceChat.h b/Engine/Plugins/Online/VoiceChat/VivoxVoiceChat/Source/Private/Android/AndroidVivoxVoiceChat.h
--- a/Engine/Plugins/Online/VoiceChat/VivoxVoiceChat/Source/Private/Android/AndroidVivoxVoiceChat.h
+++ b/Engine/Plugins/Online/VoiceChat/VivoxVoiceChat/Source/Private/Android/AndroidVivoxVoiceChat.h
@@ -12,6 +12,7 @@ public:
 
 	// ~Begin IVoiceChat Interface
 	virtual bool Initialize() override;
+	virtual bool Uninitialize() override;
 	// ~End IVoiceChat Interface
 
 private:
diff --git a/Engine/Plugins/Online/VoiceChat/VivoxVoiceChat/Source/Private/IOS/IOSVivoxVoiceChat.cpp b/Engine/Plugins/Online/VoiceChat/VivoxVoiceChat/Source/Private/IOS/IOSVivoxVoiceChat.cpp
--- a/Engine/Plugins/Online/VoiceChat/VivoxVoiceChat/Source/Private/IOS/IOSVivoxVoiceChat.cpp
+++ b/Engine/Plugins/Online/VoiceChat/VivoxVoiceChat/Source/Private/IOS/IOSVivoxVoiceChat.cpp
@@ -36,12 +36,10 @@ bool FIOSVivoxVoiceChat::Initialize()
 		{
 			ApplicationWillEnterBackgroundHandle = FCoreDelegates::ApplicationWillDeactivateDelegate.AddRaw(this, &FIOSVivoxVoiceChat::HandleApplicationWillEnterBackground);
 		}
-		//ApplicationWillEnterBackgroundHandle = FCoreDelegates::ApplicationWillDeactivateDelegate.AddRaw(this, &FIOSVivoxVoiceChat::HandleApplicationWillEnterBackground);
 		if (!ApplicationDidEnterForegroundHandle.IsValid())
 		{
 			ApplicationDidEnterForegroundHandle = FCoreDelegates::ApplicationHasReactivatedDelegate.AddRaw(this, &FIOSVivoxVoiceChat::HandleApplicationHasEnteredForeground);
 		}
-		//ApplicationDidEnterForegroundHandle = FCoreDelegates::ApplicationHasReactivatedDelegate.AddRaw(this, &FIOSVivoxVoiceChat::HandleApplicationHasEnteredForeground);
 	}
 
 	bInBackground = false;
@@ -51,6 +49,22 @@ bool FIOSVivoxVoiceChat::Initialize()
 	return bResult;
 }
 
+bool FIOSVivoxVoiceChat::Uninitialize()
+{
+	if (ApplicationWillEnterBackgroundHandle.IsValid())
+	{
+		FCoreDelegates::ApplicationWillDeactivateDelegate.Remove(ApplicationWillEnterBackgroundHandle);
+		ApplicationWillEnterBackgroundHandle.Reset();
+	}
+	if (ApplicationDidEnterForegroundHandle.IsValid())
+	{
+		FCoreDelegates::ApplicationHasReactivatedDelegate.Remove(ApplicationDidEnterForegroundHandle);
+		ApplicationDidEnterForegroundHandle.Reset();
+	}
+
+	return FVivoxVoiceChat::Uninitialize();
+}
+
 FDelegateHandle FIOSVivoxVoiceChat::StartRecording(const FOnVoiceChatRecordSamplesAvailableDelegate::FDelegate& Delegate)
 {
 	FPlatformMisc::EnableVoiceChat(true);
@@ -99,7 +113,7 @@ void FIOSVivoxVoiceChat::OnVoiceChatDisconnectComplete(const FVoiceChatResult& R
 	{
 		bShouldReconnect = true;
 	}
-	else
+	else if (IsInitialized())
 	{
 		// disconnect complete delegate fired after entering foreground
 		Reconnect();
@@ -146,13 +160,13 @@ void FIOSVivoxVoiceChat::HandleApplicationWillEnterBackground()
 void FIOSVivoxVoiceChat::HandleApplicationHasEnteredForeground()
 {
 	UE_LOG(LogVivoxVoiceChat, Log, TEXT("OnApplicationHasEnteredForegoundDelegate"));
-		
+
 	VivoxClientConnection.WillEnterForeground();
-		
+
 	if (BGTask != UIBackgroundTaskInvalid)
 	{
 		UIApplication* App = [UIApplication sharedApplication];
-		[App endBackgroundTask:BGTask];
+		[App endBackgroundTask : BGTask];
 		BGTask = UIBackgroundTaskInvalid;
 	}
 	if (bShouldReconnect)
diff --git a/Engine/Plugins/Online/VoiceChat/VivoxVoiceChat/Source/Private/IOS/IOSVivoxVoiceChat.h b/Engine/Plugins/Online/VoiceChat/VivoxVoiceChat/Source/Private/IOS/IOSVivoxVoiceChat.h
--- a/Engine/Plugins/Online/VoiceChat/VivoxVoiceChat/Source/Private/IOS/IOSVivoxVoiceChat.h
+++ b/Engine/Plugins/Online/VoiceChat/VivoxVoiceChat/Source/Private/IOS/IOSVivoxVoiceChat.h
@@ -12,6 +12,7 @@ public:
 
 	// ~Begin IVoiceChat Interface
 	virtual bool Initialize() override;
+	virtual bool Uninitialize() override;
 	virtual FDelegateHandle StartRecording(const FOnVoiceChatRecordSamplesAvailableDelegate::FDelegate& Delegate) override;
 	virtual void StopRecording(FDelegateHandle Handle) override;
 	// ~End IVoiceChat Interface
