VivoxVoiceChat plugin
- Change FOnVoiceChatBeforeRecvAudioRenderedDelegate signature so that it can be used to inject audio into the output stream
- Fix FOnVoiceChatBeforeRecvAudioRenderedDelegate not having the correct number of samples
- Change how events are dispatched on iOS so they can be processed while the application is in the background

diff --git a/Engine/Plugins/Online/VoiceChat/VivoxVoiceChat/Source/Private/IOS/IOSVivoxVoiceChat.cpp b/Engine/Plugins/Online/VoiceChat/VivoxVoiceChat/Source/Private/IOS/IOSVivoxVoiceChat.cpp
--- a/Engine/Plugins/Online/VoiceChat/VivoxVoiceChat/Source/Private/IOS/IOSVivoxVoiceChat.cpp
+++ b/Engine/Plugins/Online/VoiceChat/VivoxVoiceChat/Source/Private/IOS/IOSVivoxVoiceChat.cpp
@@ -4,6 +4,7 @@
 
 #include "Misc/ConfigCacheIni.h"
 #include "Misc/CoreDelegates.h"
+#include "Misc/EmbeddedCommunication.h"
 #include "IOS/IOSAppDelegate.h"
 
 TUniquePtr<FVivoxVoiceChat> CreateVivoxObject()
@@ -106,6 +107,20 @@ void FIOSVivoxVoiceChat::StopRecording(FDelegateHandle Handle)
 	}
 }
 
+void FIOSVivoxVoiceChat::InvokeOnUIThread(void (Func)(void* Arg0), void* Arg0)
+{
+	[FIOSAsyncTask CreateTaskWithBlock:^bool()
+	{
+		if (Func)
+		{
+			(*Func)(Arg0);
+		}
+		return true;
+	}];
+	
+	FEmbeddedCommunication::WakeGameThread();
+}
+
 void FIOSVivoxVoiceChat::onConnectCompleted(const VivoxClientApi::Uri& Server)
 {
 	EnableVoiceChat(true);
diff --git a/Engine/Plugins/Online/VoiceChat/VivoxVoiceChat/Source/Private/IOS/IOSVivoxVoiceChat.h b/Engine/Plugins/Online/VoiceChat/VivoxVoiceChat/Source/Private/IOS/IOSVivoxVoiceChat.h
--- a/Engine/Plugins/Online/VoiceChat/VivoxVoiceChat/Source/Private/IOS/IOSVivoxVoiceChat.h
+++ b/Engine/Plugins/Online/VoiceChat/VivoxVoiceChat/Source/Private/IOS/IOSVivoxVoiceChat.h
@@ -21,6 +21,7 @@ public:
 
 protected:
 	// ~Begin DebugClientApiEventHandler Interface
+	virtual void InvokeOnUIThread(void (Func)(void* Arg0), void* Arg0) override;
 	virtual void onConnectCompleted(const VivoxClientApi::Uri& Server) override;
 	virtual void onDisconnected(const VivoxClientApi::Uri& Server, const VivoxClientApi::VCSStatus& Status) override;
 	// ~End DebugClientApiEventHandler Interface
diff --git a/Engine/Plugins/Online/VoiceChat/VivoxVoiceChat/Source/Private/VivoxVoiceChat.cpp b/Engine/Plugins/Online/VoiceChat/VivoxVoiceChat/Source/Private/VivoxVoiceChat.cpp
--- a/Engine/Plugins/Online/VoiceChat/VivoxVoiceChat/Source/Private/VivoxVoiceChat.cpp
+++ b/Engine/Plugins/Online/VoiceChat/VivoxVoiceChat/Source/Private/VivoxVoiceChat.cpp
@@ -1617,12 +1617,12 @@ void FVivoxVoiceChat::onAudioUnitAfterCaptureAudioRead(const VivoxClientApi::Uri
 	if (InitialTargetUri.IsValid())
 	{
 		FScopeLock Lock(&AfterCaptureAudioReadLock);
-		OnVoiceChatAfterCaptureAudioReadDelegate.Broadcast(MakeArrayView(PcmFrames, PcmFrameCount), AudioFrameRate, ChannelsPerFrame);
+		OnVoiceChatAfterCaptureAudioReadDelegate.Broadcast(MakeArrayView(PcmFrames, PcmFrameCount * ChannelsPerFrame), AudioFrameRate, ChannelsPerFrame);
 	}
 	else
 	{
 		FScopeLock Lock(&AudioRecordLock);
-		OnVoiceChatRecordSamplesAvailableDelegate.Broadcast(MakeArrayView(PcmFrames, PcmFrameCount), AudioFrameRate, ChannelsPerFrame);
+		OnVoiceChatRecordSamplesAvailableDelegate.Broadcast(MakeArrayView(PcmFrames, PcmFrameCount * ChannelsPerFrame), AudioFrameRate, ChannelsPerFrame);
 	}
 }
 
@@ -1631,7 +1631,7 @@ void FVivoxVoiceChat::onAudioUnitBeforeCaptureAudioSent(const VivoxClientApi::Ur
 	if (InitialTargetUri.IsValid())
 	{
 		FScopeLock Lock(&BeforeCaptureAudioSentLock);
-		OnVoiceChatBeforeCaptureAudioSentDelegate.Broadcast(MakeArrayView(PcmFrames, PcmFrameCount), AudioFrameRate, ChannelsPerFrame, bSpeaking);
+		OnVoiceChatBeforeCaptureAudioSentDelegate.Broadcast(MakeArrayView(PcmFrames, PcmFrameCount * ChannelsPerFrame), AudioFrameRate, ChannelsPerFrame, bSpeaking);
 	}
 }
 
@@ -1640,7 +1640,7 @@ void FVivoxVoiceChat::onAudioUnitBeforeRecvAudioRendered(const VivoxClientApi::U
 	if (InitialTargetUri.IsValid())
 	{
 		FScopeLock Lock(&BeforeRecvAudioRenderedLock);
-		OnVoiceChatBeforeRecvAudioRenderedDelegate.Broadcast(MakeArrayView(PcmFrames, PcmFrameCount), AudioFrameRate, ChannelsPerFrame, bSilence);
+		OnVoiceChatBeforeRecvAudioRenderedDelegate.Broadcast(MakeArrayView(PcmFrames, PcmFrameCount * ChannelsPerFrame), AudioFrameRate, ChannelsPerFrame, bSilence);
 	}
 }
 
diff --git a/Engine/Source/Runtime/Online/Voice/Public/Interfaces/VoiceChat.h b/Engine/Source/Runtime/Online/Voice/Public/Interfaces/VoiceChat.h
--- a/Engine/Source/Runtime/Online/Voice/Public/Interfaces/VoiceChat.h
+++ b/Engine/Source/Runtime/Online/Voice/Public/Interfaces/VoiceChat.h
@@ -101,7 +101,7 @@ DECLARE_MULTICAST_DELEGATE_ThreeParams(FOnVoiceChatPlayerVolumeUpdatedDelegate,
 DECLARE_MULTICAST_DELEGATE_ThreeParams(FOnVoiceChatRecordSamplesAvailableDelegate, TArrayView<const int16> /* PcmSamples */, int /* SampleRate */, int /* Channels */);
 DECLARE_MULTICAST_DELEGATE_ThreeParams(FOnVoiceChatAfterCaptureAudioReadDelegate, TArrayView<int16> /* PcmSamples */, int /* SampleRate */, int /* Channels */);
 DECLARE_MULTICAST_DELEGATE_FourParams(FOnVoiceChatBeforeCaptureAudioSentDelegate, TArrayView<const int16> /* PcmSamples */, int /* SampleRate */, int /* Channels */, bool /* bIsSpeaking */);
-DECLARE_MULTICAST_DELEGATE_FourParams(FOnVoiceChatBeforeRecvAudioRenderedDelegate, TArrayView<const int16> /* PcmSamples */, int /* SampleRate */, int /* Channels */, bool /* bIsSilence */);
+DECLARE_MULTICAST_DELEGATE_FourParams(FOnVoiceChatBeforeRecvAudioRenderedDelegate, TArrayView<int16> /* PcmSamples */, int /* SampleRate */, int /* Channels */, bool /* bIsSilence */);
 
 class IVoiceChat : public IModularFeature
 {
