VivoxVoiceChat plugin - On iOS, ensure hardware AEC is disabled if we are using an A2DP bluetooth headset.

diff --git a/Engine/Plugins/Online/VoiceChat/VivoxVoiceChat/Source/Private/IOS/IOSVivoxVoiceChat.cpp b/Engine/Plugins/Online/VoiceChat/VivoxVoiceChat/Source/Private/IOS/IOSVivoxVoiceChat.cpp
--- a/Engine/Plugins/Online/VoiceChat/VivoxVoiceChat/Source/Private/IOS/IOSVivoxVoiceChat.cpp
+++ b/Engine/Plugins/Online/VoiceChat/VivoxVoiceChat/Source/Private/IOS/IOSVivoxVoiceChat.cpp
@@ -45,6 +45,10 @@ bool FIOSVivoxVoiceChat::Initialize()
 		{
 			ApplicationDidEnterForegroundHandle = FCoreDelegates::ApplicationHasReactivatedDelegate.AddRaw(this, &FIOSVivoxVoiceChat::HandleApplicationHasEnteredForeground);
 		}
+		if (!AudioRouteChangedHandle.IsValid())
+		{
+			AudioRouteChangedHandle = FCoreDelegates::AudioRouteChangedDelegate.AddRaw(this, &FIOSVivoxVoiceChat::HandleAudioRouteChanged);
+		}
 	}
 
 	bInBackground = false;
@@ -66,6 +70,11 @@ bool FIOSVivoxVoiceChat::Uninitialize()
 		FCoreDelegates::ApplicationHasReactivatedDelegate.Remove(ApplicationDidEnterForegroundHandle);
 		ApplicationDidEnterForegroundHandle.Reset();
 	}
+	if (AudioRouteChangedHandle.IsValid())
+	{
+		FCoreDelegates::AudioRouteChangedDelegate.Remove(AudioRouteChangedHandle);
+		AudioRouteChangedHandle.Reset();
+	}
 
 	return FVivoxVoiceChat::Uninitialize();
 }
@@ -121,10 +130,28 @@ void FIOSVivoxVoiceChat::InvokeOnUIThread(void (Func)(void* Arg0), void* Arg0)
 	FEmbeddedCommunication::WakeGameThread();
 }
 
-void FIOSVivoxVoiceChat::onConnectCompleted(const VivoxClientApi::Uri& Server)
+void FIOSVivoxVoiceChat::onChannelJoined(const VivoxClientApi::AccountName& AccountName, const VivoxClientApi::Uri& ChannelUri)
 {
 	EnableVoiceChat(true);
-	FVivoxVoiceChat::onConnectCompleted(Server);
+	FVivoxVoiceChat::onChannelJoined(AccountName, ChannelUri);
+}
+
+void FIOSVivoxVoiceChat::onChannelExited(const VivoxClientApi::AccountName& AccountName, const VivoxClientApi::Uri& ChannelUri, const VivoxClientApi::VCSStatus& Status)
+{
+	FVivoxVoiceChat::onChannelExited(AccountName, ChannelUri, Status);
+	
+	bool bIsInChannel = false;
+	for (const TPair<FString,FChannelSession>& ChannelSession : LoginSession.ChannelSessions)
+	{
+		if (ChannelSession.Value.State == FChannelSession::EState::Connected)
+		{
+			bIsInChannel = true;
+		}
+	}
+	if (!bIsRecording && !bIsInChannel)
+	{
+		EnableVoiceChat(false);
+	}
 }
 
 void FIOSVivoxVoiceChat::onDisconnected(const VivoxClientApi::Uri& Server, const VivoxClientApi::VCSStatus& Status)
@@ -246,6 +273,15 @@ void FIOSVivoxVoiceChat::HandleApplicationHasEnteredForeground()
 	{
 		Reconnect();
 	}
+
+	// HandleAudioRouteChanged is not getting called when a route change happens in the background. Update voice chat settings here to handle this case
+	UpdateVoiceChatSettings();
+}
+
+void FIOSVivoxVoiceChat::HandleAudioRouteChanged(bool)
+{
+	UE_LOG(LogVivoxVoiceChat, Verbose, TEXT("Audio route changed"));
+	UpdateVoiceChatSettings();
 }
 
 void FIOSVivoxVoiceChat::Reconnect()
@@ -263,7 +299,7 @@ void FIOSVivoxVoiceChat::EnableVoiceChat(bool bEnable)
 {
 	if (FPlatformMisc::IsVoiceChatEnabled() != bEnable)
 	{
-		if (IsHardwareAECEnabled())
+		if (IsHardwareAECEnabled() && !IsBluetoothA2DPInUse())
 		{
 			[[IOSAppDelegate GetDelegate] EnableHighQualityVoiceChat:bEnable];
 			vx_set_platform_aec_enabled(bEnable ? 1 : 0);
@@ -277,7 +313,8 @@ void FIOSVivoxVoiceChat::UpdateVoiceChatSettings()
 	if (FPlatformMisc::IsVoiceChatEnabled())
 	{
 		// update the aec settings
-		const bool bEnableAEC = IsHardwareAECEnabled();
+		const bool bEnableAEC = IsHardwareAECEnabled() && !IsBluetoothA2DPInUse();
+		UE_LOG(LogVivoxVoiceChat, Verbose, TEXT("%s AEC"), bEnableAEC ? TEXT("Enabling") : TEXT("Disabling"))
 		[[IOSAppDelegate GetDelegate] EnableHighQualityVoiceChat:bEnableAEC];
 		vx_set_platform_aec_enabled(bEnableAEC ? 1 : 0);
 		// reenable voice chat to apply the changes
@@ -285,3 +322,18 @@ void FIOSVivoxVoiceChat::UpdateVoiceChatSettings()
 	}
 }
 
+bool FIOSVivoxVoiceChat::IsBluetoothA2DPInUse()
+{
+	if (AVAudioSessionRouteDescription* CurrentRoute = [[AVAudioSession sharedInstance] currentRoute])
+	{
+		for (AVAudioSessionPortDescription* Port in [CurrentRoute outputs])
+		{
+			if ([[Port portType] isEqualToString:AVAudioSessionPortBluetoothA2DP])
+			{
+				return true;
+			}
+		}
+	}
+
+	return false;
+}
diff --git a/Engine/Plugins/Online/VoiceChat/VivoxVoiceChat/Source/Private/IOS/IOSVivoxVoiceChat.h b/Engine/Plugins/Online/VoiceChat/VivoxVoiceChat/Source/Private/IOS/IOSVivoxVoiceChat.h
--- a/Engine/Plugins/Online/VoiceChat/VivoxVoiceChat/Source/Private/IOS/IOSVivoxVoiceChat.h
+++ b/Engine/Plugins/Online/VoiceChat/VivoxVoiceChat/Source/Private/IOS/IOSVivoxVoiceChat.h
@@ -22,7 +22,8 @@ public:
 protected:
 	// ~Begin DebugClientApiEventHandler Interface
 	virtual void InvokeOnUIThread(void (Func)(void* Arg0), void* Arg0) override;
-	virtual void onConnectCompleted(const VivoxClientApi::Uri& Server) override;
+	virtual void onChannelJoined(const VivoxClientApi::AccountName& AccountName, const VivoxClientApi::Uri& ChannelUri) override;
+	virtual void onChannelExited(const VivoxClientApi::AccountName& AccountName, const VivoxClientApi::Uri& ChannelUri, const VivoxClientApi::VCSStatus& Status) override;
 	virtual void onDisconnected(const VivoxClientApi::Uri& Server, const VivoxClientApi::VCSStatus& Status) override;
 	// ~End DebugClientApiEventHandler Interface
 	
@@ -35,11 +36,13 @@ private:
 
 	void HandleApplicationWillEnterBackground();
 	void HandleApplicationHasEnteredForeground();
+	void HandleAudioRouteChanged(bool);
 
 	void Reconnect();
 
 	FDelegateHandle ApplicationWillEnterBackgroundHandle;
 	FDelegateHandle ApplicationDidEnterForegroundHandle;
+	FDelegateHandle AudioRouteChangedHandle;
 
 	UIBackgroundTaskIdentifier BGTask;
 	bool bDisconnectInBackground;
@@ -57,5 +60,5 @@ private:
 
 	void EnableVoiceChat(bool bEnable);
 	void UpdateVoiceChatSettings();
-	void DisableVoiceChat();
+	bool IsBluetoothA2DPInUse();
 };
