[NO NON-PII COMMIT MESSAGE]

diff --git a/Engine/Plugins/Online/VoiceChat/VivoxVoiceChat/Source/Private/VivoxVoiceChat.cpp b/Engine/Plugins/Online/VoiceChat/VivoxVoiceChat/Source/Private/VivoxVoiceChat.cpp
--- a/Engine/Plugins/Online/VoiceChat/VivoxVoiceChat/Source/Private/VivoxVoiceChat.cpp
+++ b/Engine/Plugins/Online/VoiceChat/VivoxVoiceChat/Source/Private/VivoxVoiceChat.cpp
@@ -59,7 +59,7 @@ FVoiceChatResult ResultFromVivoxStatus(const VivoxClientApi::VCSStatus& Status)
 	return { !Status.IsError(), Status.GetStatusCode(), ErrorString };
 }
 
-FVoiceChatResult ResultFromErrorString(const FString& Error, int ErrorCode = -1)
+FVoiceChatResult ResultFromErrorString(const FString& Error, int ErrorCode)
 {
 	return { false, ErrorCode, Error };
 }
@@ -488,7 +488,7 @@ void FVivoxVoiceChat::Connect(const FOnVoiceChatConnectCompleteDelegate& Delegat
 	}
 	else if (ConnectionState == EConnectionState::Disconnecting)
 	{
-		Result = ResultFromErrorString("Disconnect in progress");
+		Result = ResultFromErrorString("Disconnect in progress", -1);
 	}
 
 	if (!Result.bSuccess)
@@ -510,7 +510,7 @@ void FVivoxVoiceChat::Connect(const FOnVoiceChatConnectCompleteDelegate& Delegat
 			if (VivoxServerUrl.IsEmpty() || VivoxDomain.IsEmpty() || VivoxNamespace.IsEmpty())
 			{
 				UE_LOG(LogVivoxVoiceChat, Warning, TEXT("[VoiceChat.Vivox] ServerUrl, Domain, or Issuer is not set. Vivox voice chat will not work"));
-				Result = ResultFromErrorString(TEXT("Vivox config missing"));
+				Result = ResultFromErrorString(TEXT("Vivox config missing"), -1);
 			}
 			else
 			{
@@ -574,7 +574,7 @@ void FVivoxVoiceChat::Login(FPlatformUserId PlatformId, const FString& PlayerNam
 	}
 	else if (!IsConnected())
 	{
-		Result = ResultFromErrorString(TEXT("Not Connected"));
+		Result = ResultFromErrorString(TEXT("Not Connected"), -2);
 	}
 	else if (IsLoggedIn())
 	{
@@ -599,7 +599,7 @@ void FVivoxVoiceChat::Login(FPlatformUserId PlatformId, const FString& PlayerNam
 	else if (PlayerName.Len() > 60 - VivoxNamespace.Len())
 	{
 		// Name must be between 3-63 characters long and start and end with a '.'. It also must contain the issuer and another '.' separating issuer and name
-		Result = ResultFromErrorString(TEXT("PlayerName is too long"));
+		Result = ResultFromErrorString(TEXT("PlayerName is too long"), -1);
 	}
 
 	if (!Result.bSuccess)
@@ -649,7 +649,7 @@ void FVivoxVoiceChat::Logout(const FOnVoiceChatLogoutCompleteDelegate& Delegate)
 	}
 	else if (!IsConnected())
 	{
-		Result = ResultFromErrorString(TEXT("Not Connected"));
+		Result = ResultFromErrorString(TEXT("Not Connected"), -2);
 	}
 	else if (!IsLoggedIn())
 	{
@@ -746,7 +746,7 @@ void FVivoxVoiceChat::JoinChannel(const FString& ChannelName, const FString& Cha
 	}
 	else if (!IsConnected())
 	{
-		Result = ResultFromErrorString(TEXT("Not Connected"));
+		Result = ResultFromErrorString(TEXT("Not Connected"), -2);
 	}
 	else if (!IsLoggedIn())
 	{
@@ -781,12 +781,12 @@ void FVivoxVoiceChat::JoinChannel(const FString& ChannelName, const FString& Cha
 	}
 	else if (ChannelSession.State == FChannelSession::EState::Connecting)
 	{
-		Delegate.ExecuteIfBound(ChannelName, ResultFromErrorString("Join in progress"));
+		Delegate.ExecuteIfBound(ChannelName, ResultFromErrorString("Join in progress", -1));
 		return;
 	}
 	else if (ChannelSession.State == FChannelSession::EState::Disconnecting)
 	{
-		Delegate.ExecuteIfBound(ChannelName, ResultFromErrorString("Leave in progress"));
+		Delegate.ExecuteIfBound(ChannelName, ResultFromErrorString("Leave in progress", -1));
 		return;
 	}
 
@@ -816,7 +816,7 @@ void FVivoxVoiceChat::LeaveChannel(const FString& ChannelName, const FOnVoiceCha
 	}
 	else if (!IsConnected())
 	{
-		Result = ResultFromErrorString(TEXT("Not Connected"));
+		Result = ResultFromErrorString(TEXT("Not Connected"), -2);
 	}
 	else if (!IsLoggedIn())
 	{
@@ -830,7 +830,7 @@ void FVivoxVoiceChat::LeaveChannel(const FString& ChannelName, const FOnVoiceCha
 	FChannelSession& ChannelSession = GetChannelSession(ChannelName);
 	if (ChannelSession.State != FChannelSession::EState::Connected)
 	{
-		Result = ResultFromErrorString("Not in channel");
+		Result = ResultFromErrorString("Not in channel", -1);
 	}
 
 	if (!Result.bSuccess)
@@ -1276,7 +1276,7 @@ void FVivoxVoiceChat::onInvalidLoginCredentials(const VivoxClientApi::AccountNam
 
 	FString PlayerName = GetPlayerNameFromAccountName(AccountName);
 
-	TriggerCompletionDelegate(OnVoiceChatLoginCompleteDelegate, PlayerName, ResultFromErrorString(TEXT("Invalid login credentials")));
+	TriggerCompletionDelegate(OnVoiceChatLoginCompleteDelegate, PlayerName, ResultFromErrorString(TEXT("Invalid login credentials"), -1));
 }
 
 void FVivoxVoiceChat::onLoginFailed(const VivoxClientApi::AccountName& AccountName, const VivoxClientApi::VCSStatus& Status)
@@ -1331,7 +1331,7 @@ void FVivoxVoiceChat::onInvalidChannelCredentials(const VivoxClientApi::AccountN
 
 	FChannelSession& ChannelSession = GetChannelSession(ChannelUri);
 	ChannelSession.State = FChannelSession::EState::Disconnected;
-	TriggerCompletionDelegate(ChannelSession.JoinDelegate, ChannelSession.ChannelName, ResultFromErrorString(TEXT("Invalid join credentials")));
+	TriggerCompletionDelegate(ChannelSession.JoinDelegate, ChannelSession.ChannelName, ResultFromErrorString(TEXT("Invalid join credentials"), -1));
 
 	RemoveChannelSession(ChannelSession.ChannelName);
 }
