[Vivox] Remove unnecessary bool

diff --git a/Engine/Plugins/Online/VoiceChat/VivoxVoiceChat/Source/Private/VivoxVoiceChat.cpp b/Engine/Plugins/Online/VoiceChat/VivoxVoiceChat/Source/Private/VivoxVoiceChat.cpp
--- a/Engine/Plugins/Online/VoiceChat/VivoxVoiceChat/Source/Private/VivoxVoiceChat.cpp
+++ b/Engine/Plugins/Online/VoiceChat/VivoxVoiceChat/Source/Private/VivoxVoiceChat.cpp
@@ -1359,7 +1359,6 @@ void FVivoxVoiceChat::onChannelExited(const VivoxClientApi::AccountName& Account
 	}
 
 	FChannelSession& ChannelSession = GetChannelSession(ChannelUri);
-	const bool bWasConnected = ChannelSession.State == FChannelSession::EState::Connected;
 	const bool bWasConnecting = ChannelSession.State == FChannelSession::EState::Connecting;
 	const bool bWasDisconnecting = ChannelSession.State == FChannelSession::EState::Disconnecting;
 	ChannelSession.State = FChannelSession::EState::Disconnected;
