VivoxVoiceChat plugin - preliminary support for iOS bluetooth input. Audio quality is poor and will require updates to the vivox sdk to resolve

diff --git a/Engine/Plugins/Online/VoiceChat/VivoxVoiceChat/Source/Private/IOS/IOSVivoxVoiceChat.cpp b/Engine/Plugins/Online/VoiceChat/VivoxVoiceChat/Source/Private/IOS/IOSVivoxVoiceChat.cpp
--- a/Engine/Plugins/Online/VoiceChat/VivoxVoiceChat/Source/Private/IOS/IOSVivoxVoiceChat.cpp
+++ b/Engine/Plugins/Online/VoiceChat/VivoxVoiceChat/Source/Private/IOS/IOSVivoxVoiceChat.cpp
@@ -39,7 +39,8 @@ bool FIOSVivoxVoiceChat::Initialize()
 		GConfig->GetBool(TEXT("VoiceChat.Vivox"), TEXT("bDisconnectInBackground"), bDisconnectInBackground, GEngineIni);
 		GConfig->GetFloat(TEXT("VoiceChat.Vivox"), TEXT("BackgroundDelayedDisconnectTime"), BackgroundDelayedDisconnectTime, GEngineIni);
 		GConfig->GetBool(TEXT("VoiceChat.Vivox"), TEXT("bEnableHardwareAEC"), bEnableHardwareAEC, GEngineIni);
-		
+		GConfig->GetBool(TEXT("VoiceChat.Vivox"), TEXT("bEnableBluetoothMicrophone"), bEnableBluetoothMicrophone, GEngineIni);
+
 		if (!ApplicationWillEnterBackgroundHandle.IsValid())
 		{
 			ApplicationWillEnterBackgroundHandle = FCoreDelegates::ApplicationWillDeactivateDelegate.AddRaw(this, &FIOSVivoxVoiceChat::HandleApplicationWillEnterBackground);
@@ -52,6 +53,8 @@ bool FIOSVivoxVoiceChat::Initialize()
 		{
 			AudioRouteChangedHandle = FCoreDelegates::AudioRouteChangedDelegate.AddRaw(this, &FIOSVivoxVoiceChat::HandleAudioRouteChanged);
 		}
+
+		UpdateVoiceChatSettings();
 	}
 
 	bInBackground = false;
@@ -98,6 +101,11 @@ void FIOSVivoxVoiceChat::SetSetting(const FString& Name, const FString& Value)
 		OverrideEnableHardwareAEC = FCString::ToBool(*Value);
 		UpdateVoiceChatSettings();
 	}
+	else if (Name == TEXT("BluetoothMicrophone"))
+	{
+		OverrideEnableBluetoothMicrophone = FCString::ToBool(*Value);
+		UpdateVoiceChatSettings();
+	}
 	else
 	{
 		FVivoxVoiceChat::SetSetting(Name, Value);
@@ -108,7 +116,11 @@ FString FIOSVivoxVoiceChat::GetSetting(const FString& Name)
 {
 	if (Name == TEXT("HardwareAEC"))
 	{
-		return LexToString(OverrideEnableHardwareAEC.Get(bEnableHardwareAEC));
+		return LexToString(IsHardwareAECEnabled());
+	}
+	else if (Name == TEXT("BluetoothMicrophone"))
+	{
+		return LexToString(IsBluetoothMicrophoneEnabled());
 	}
 	return FVivoxVoiceChat::GetSetting(Name);
 }
@@ -295,6 +307,11 @@ bool FIOSVivoxVoiceChat::IsHardwareAECEnabled() const
 	return OverrideEnableHardwareAEC.Get(bEnableHardwareAEC);
 }
 
+bool FIOSVivoxVoiceChat::IsBluetoothMicrophoneEnabled() const
+{
+	return OverrideEnableBluetoothMicrophone.Get(bEnableBluetoothMicrophone);
+}
+
 void FIOSVivoxVoiceChat::EnableVoiceChat(bool bEnable)
 {
 	if (bEnable)
@@ -302,9 +319,10 @@ void FIOSVivoxVoiceChat::EnableVoiceChat(bool bEnable)
 		++VoiceChatEnableCount;
 		if (VoiceChatEnableCount == 1)
 		{
-			const bool bEnableAEC = IsHardwareAECEnabled() && !IsBluetoothA2DPInUse();
 			[[IOSAppDelegate GetDelegate] SetFeature:EAudioFeature::Playback Active:true];
 			[[IOSAppDelegate GetDelegate] SetFeature:EAudioFeature::Record Active:true];
+
+			const bool bEnableAEC = IsHardwareAECEnabled() && IsUsingBuiltInSpeaker();
 			vx_set_platform_aec_enabled(bEnableAEC ? 1 : 0);
 		}
 	}
@@ -325,7 +343,13 @@ void FIOSVivoxVoiceChat::EnableVoiceChat(bool bEnable)
 
 void FIOSVivoxVoiceChat::UpdateVoiceChatSettings()
 {
-	const bool bEnableAEC = IsHardwareAECEnabled() && !IsBluetoothA2DPInUse();
+	if (bBluetoothMicrophoneFeatureEnabled != IsBluetoothMicrophoneEnabled())
+	{
+		[[IOSAppDelegate GetDelegate] SetFeature:EAudioFeature::BluetoothMicrophone Active:IsBluetoothMicrophoneEnabled()];
+		bBluetoothMicrophoneFeatureEnabled = IsBluetoothMicrophoneEnabled();
+	}
+
+	const bool bEnableAEC = IsHardwareAECEnabled() && IsUsingBuiltInSpeaker();
 
 	if (bVoiceChatModeEnabled != bEnableAEC)
 	{
@@ -340,13 +364,14 @@ void FIOSVivoxVoiceChat::UpdateVoiceChatSettings()
 	}
 }
 
-bool FIOSVivoxVoiceChat::IsBluetoothA2DPInUse()
+bool FIOSVivoxVoiceChat::IsUsingBuiltInSpeaker()
 {
 	if (AVAudioSessionRouteDescription* CurrentRoute = [[AVAudioSession sharedInstance] currentRoute])
 	{
 		for (AVAudioSessionPortDescription* Port in [CurrentRoute outputs])
 		{
-			if ([[Port portType] isEqualToString:AVAudioSessionPortBluetoothA2DP])
+			if ([[Port portType] isEqualToString:AVAudioSessionPortBuiltInReceiver]
+				|| [[Port portType] isEqualToString:AVAudioSessionPortBuiltInSpeaker])
 			{
 				return true;
 			}
diff --git a/Engine/Plugins/Online/VoiceChat/VivoxVoiceChat/Source/Private/IOS/IOSVivoxVoiceChat.h b/Engine/Plugins/Online/VoiceChat/VivoxVoiceChat/Source/Private/IOS/IOSVivoxVoiceChat.h
--- a/Engine/Plugins/Online/VoiceChat/VivoxVoiceChat/Source/Private/IOS/IOSVivoxVoiceChat.h
+++ b/Engine/Plugins/Online/VoiceChat/VivoxVoiceChat/Source/Private/IOS/IOSVivoxVoiceChat.h
@@ -54,13 +54,19 @@ private:
 	float BackgroundDelayedDisconnectTime;
 	NSTimer* DelayedDisconnectTimer;
 
+	uint VoiceChatEnableCount = 0;
+
 	bool IsHardwareAECEnabled() const;
 	bool bEnableHardwareAEC = false;
 	TOptional<bool> OverrideEnableHardwareAEC;
 	bool bVoiceChatModeEnabled = false;
-	uint VoiceChatEnableCount = 0;
+
+	bool IsBluetoothMicrophoneEnabled() const;
+	bool bEnableBluetoothMicrophone = false;
+	TOptional<bool> OverrideEnableBluetoothMicrophone;
+	bool bBluetoothMicrophoneFeatureEnabled = false;
 
 	void EnableVoiceChat(bool bEnable);
 	void UpdateVoiceChatSettings();
-	bool IsBluetoothA2DPInUse();
+	bool IsUsingBuiltInSpeaker();
 };
