Don't make unnecessary AVAudioSession setCategory calls
VivoxVoiceChat - enable voice chat before calling JoinChannel

diff --git a/Engine/Plugins/Online/VoiceChat/VivoxVoiceChat/Source/Private/IOS/IOSVivoxVoiceChat.cpp b/Engine/Plugins/Online/VoiceChat/VivoxVoiceChat/Source/Private/IOS/IOSVivoxVoiceChat.cpp
--- a/Engine/Plugins/Online/VoiceChat/VivoxVoiceChat/Source/Private/IOS/IOSVivoxVoiceChat.cpp
+++ b/Engine/Plugins/Online/VoiceChat/VivoxVoiceChat/Source/Private/IOS/IOSVivoxVoiceChat.cpp
@@ -29,9 +29,6 @@ FIOSVivoxVoiceChat::~FIOSVivoxVoiceChat()
 
 bool FIOSVivoxVoiceChat::Initialize()
 {
-	// Vivox sdk opens the microphone if we are in one of the Ambient modes and the user backgrounds the application. Set the mode to Playback to work around this
-	[[IOSAppDelegate GetDelegate] SetFeature:EAudioFeature::Playback Active:true];
-
 	bool bResult = FVivoxVoiceChat::Initialize();
 
 	if (bResult)
@@ -89,8 +86,6 @@ bool FIOSVivoxVoiceChat::Uninitialize()
 		EnableVoiceChat(false);
 	}
 
-	[[IOSAppDelegate GetDelegate] SetFeature:EAudioFeature::Playback Active:false];
-
 	return bReturn;
 }
 
@@ -125,6 +120,23 @@ FString FIOSVivoxVoiceChat::GetSetting(const FString& Name)
 	return FVivoxVoiceChat::GetSetting(Name);
 }
 
+void FIOSVivoxVoiceChat::JoinChannel(const FString& ChannelName, const FString& ChannelCredentials, EVoiceChatChannelType ChannelType, const FOnVoiceChatChannelJoinCompleteDelegate& Delegate, TOptional<FVoiceChatChannel3dProperties> Channel3dProperties)
+{
+	EnableVoiceChat(true);
+
+	FOnVoiceChatChannelJoinCompleteDelegate DelegateWrapper = FOnVoiceChatChannelJoinCompleteDelegate::CreateLambda(
+		[this, Delegate](const FString& ChannelName, const FVoiceChatResult& Result)
+		{
+			if (!Result.bSuccess)
+			{
+				EnableVoiceChat(false);
+			}
+			Delegate.ExecuteIfBound(ChannelName, Result);
+		});
+
+	FVivoxVoiceChat::JoinChannel(ChannelName, ChannelCredentials, ChannelType, DelegateWrapper, Channel3dProperties);
+}
+
 FDelegateHandle FIOSVivoxVoiceChat::StartRecording(const FOnVoiceChatRecordSamplesAvailableDelegate::FDelegate& Delegate)
 {
 	bIsRecording = true;
@@ -155,7 +167,6 @@ void FIOSVivoxVoiceChat::InvokeOnUIThread(void (Func)(void* Arg0), void* Arg0)
 
 void FIOSVivoxVoiceChat::onChannelJoined(const VivoxClientApi::AccountName& AccountName, const VivoxClientApi::Uri& ChannelUri)
 {
-	EnableVoiceChat(true);
 	FVivoxVoiceChat::onChannelJoined(AccountName, ChannelUri);
 }
 
diff --git a/Engine/Plugins/Online/VoiceChat/VivoxVoiceChat/Source/Private/IOS/IOSVivoxVoiceChat.h b/Engine/Plugins/Online/VoiceChat/VivoxVoiceChat/Source/Private/IOS/IOSVivoxVoiceChat.h
--- a/Engine/Plugins/Online/VoiceChat/VivoxVoiceChat/Source/Private/IOS/IOSVivoxVoiceChat.h
+++ b/Engine/Plugins/Online/VoiceChat/VivoxVoiceChat/Source/Private/IOS/IOSVivoxVoiceChat.h
@@ -15,6 +15,7 @@ public:
 	virtual bool Uninitialize() override;
 	virtual void SetSetting(const FString& Name, const FString& Value) override;
 	virtual FString GetSetting(const FString& Name) override;
+	virtual void JoinChannel(const FString& ChannelName, const FString& ChannelCredentials, EVoiceChatChannelType ChannelType, const FOnVoiceChatChannelJoinCompleteDelegate& Delegate, TOptional<FVoiceChatChannel3dProperties> Channel3dProperties = TOptional<FVoiceChatChannel3dProperties>()) override;
 	virtual FDelegateHandle StartRecording(const FOnVoiceChatRecordSamplesAvailableDelegate::FDelegate& Delegate) override;
 	virtual void StopRecording(FDelegateHandle Handle) override;
 	// ~End IVoiceChat Interface
diff --git a/Engine/Source/Runtime/ApplicationCore/Private/IOS/IOSAppDelegate.cpp b/Engine/Source/Runtime/ApplicationCore/Private/IOS/IOSAppDelegate.cpp
--- a/Engine/Source/Runtime/ApplicationCore/Private/IOS/IOSAppDelegate.cpp
+++ b/Engine/Source/Runtime/ApplicationCore/Private/IOS/IOSAppDelegate.cpp
@@ -681,15 +681,24 @@ static IOSAppDelegate* CachedDelegate = nil;
 		Category = AVAudioSessionCategoryAmbient;
 	}
 
+	
+	AVAudioSessionCategoryOptions TestOptions = Options;
+	if ([[[AVAudioSession sharedInstance] category] compare:AVAudioSessionCategoryAmbient] == NSOrderedSame)
+	{
+		// This option is implicitly set, so we need to compare the existing options with it
+		TestOptions |= AVAudioSessionCategoryOptionMixWithOthers;
+	}
 	// set the category if anything has changed
 	if ([[[AVAudioSession sharedInstance] category] compare:Category] != NSOrderedSame ||
 		[[[AVAudioSession sharedInstance] mode] compare:Mode] != NSOrderedSame ||
-		[[AVAudioSession sharedInstance] categoryOptions] != Options)
+		[[AVAudioSession sharedInstance] categoryOptions] != TestOptions)
 	{
+		UE_LOG(LogIOSAudioSession, Log, TEXT("Setting AVAudioSession category to Category:%s Mode:%s Options:%x"), *FString(Category), *FString(Mode), Options);
+
 		[[AVAudioSession sharedInstance] setCategory:Category mode:Mode options:Options error:&ActiveError];
 		if (ActiveError)
 		{
-			UE_LOG(LogIOSAudioSession, Error, TEXT("Failed to set audio session category to Category:%s Mode:%s Options:%x! [Error = %s]"), *FString(Category), *FString(Mode), Options, *FString([ActiveError description]));
+			UE_LOG(LogIOSAudioSession, Error, TEXT("Failed to set AVAudioSession category to Category:%s Mode:%s Options:%x! [Error = %s]"), *FString(Category), *FString(Mode), Options, *FString([ActiveError description]));
 		}
 	}
 }
