VivoxVoiceChat - Work around an sdk issue that causes the vivox sdk to open the microphone on iOS when the user backgrounds the application

diff --git a/Engine/Plugins/Online/VoiceChat/VivoxVoiceChat/Source/Private/IOS/IOSVivoxVoiceChat.cpp b/Engine/Plugins/Online/VoiceChat/VivoxVoiceChat/Source/Private/IOS/IOSVivoxVoiceChat.cpp
--- a/Engine/Plugins/Online/VoiceChat/VivoxVoiceChat/Source/Private/IOS/IOSVivoxVoiceChat.cpp
+++ b/Engine/Plugins/Online/VoiceChat/VivoxVoiceChat/Source/Private/IOS/IOSVivoxVoiceChat.cpp
@@ -29,6 +29,9 @@ FIOSVivoxVoiceChat::~FIOSVivoxVoiceChat()
 
 bool FIOSVivoxVoiceChat::Initialize()
 {
+	// Vivox sdk opens the microphone if we are in one of the Ambient modes and the user backgrounds the application. Set the mode to Playback to work around this
+	[[IOSAppDelegate GetDelegate] SetFeature:EAudioFeature::Playback Active:true];
+
 	bool bResult = FVivoxVoiceChat::Initialize();
 
 	if (bResult)
@@ -86,6 +89,8 @@ bool FIOSVivoxVoiceChat::Uninitialize()
 		EnableVoiceChat(false);
 	}
 
+	[[IOSAppDelegate GetDelegate] SetFeature:EAudioFeature::Playback Active:false];
+
 	return bReturn;
 }
 
