Update iOS Vivox SDK to 4.9.0002.32243. This fixes an issue where the recording icon would show up while not in a voice channel

diff --git a/Engine/Plugins/Online/VoiceChat/VivoxVoiceChat/Source/Private/IOS/IOSVivoxVoiceChat.cpp b/Engine/Plugins/Online/VoiceChat/VivoxVoiceChat/Source/Private/IOS/IOSVivoxVoiceChat.cpp
--- a/Engine/Plugins/Online/VoiceChat/VivoxVoiceChat/Source/Private/IOS/IOSVivoxVoiceChat.cpp
+++ b/Engine/Plugins/Online/VoiceChat/VivoxVoiceChat/Source/Private/IOS/IOSVivoxVoiceChat.cpp
@@ -29,9 +29,6 @@ FIOSVivoxVoiceChat::~FIOSVivoxVoiceChat()
 
 bool FIOSVivoxVoiceChat::Initialize()
 {
-	// Vivox sdk opens the microphone if we are in one of the Ambient modes and the user backgrounds the application. Set the mode to Playback to work around this
-	[[IOSAppDelegate GetDelegate] SetFeature:EAudioFeature::Playback Active:true];
-
 	bool bResult = FVivoxVoiceChat::Initialize();
 
 	if (bResult)
@@ -89,8 +86,6 @@ bool FIOSVivoxVoiceChat::Uninitialize()
 		EnableVoiceChat(false);
 	}
 
-	[[IOSAppDelegate GetDelegate] SetFeature:EAudioFeature::Playback Active:false];
-
 	return bReturn;
 }
 
