VivoxVoiceChat - Expose rtp timeout settings (new in SDK 4.10.0000.32327) via ini

diff --git a/Engine/Plugins/Online/VoiceChat/VivoxVoiceChat/Source/Private/VivoxVoiceChat.cpp b/Engine/Plugins/Online/VoiceChat/VivoxVoiceChat/Source/Private/VivoxVoiceChat.cpp
--- a/Engine/Plugins/Online/VoiceChat/VivoxVoiceChat/Source/Private/VivoxVoiceChat.cpp
+++ b/Engine/Plugins/Online/VoiceChat/VivoxVoiceChat/Source/Private/VivoxVoiceChat.cpp
@@ -1655,6 +1655,11 @@ void FVivoxVoiceChat::SetVivoxSdkConfigHints(vx_sdk_config_t& Hints)
 	bool bEnableAudioDucking = false;
 	GConfig->GetBool(TEXT("VoiceChat.Vivox"), TEXT("bEnableAudioDucking"), bEnableAudioDucking, GEngineIni);
 	Hints.disable_audio_ducking = !bEnableAudioDucking;
+
+	GConfig->GetInt(TEXT("VoiceChat.Vivox"), TEXT("RtpConnectTimeoutMs"), Hints.never_rtp_timeout_ms, GEngineIni);
+	GConfig->GetInt(TEXT("VoiceChat.Vivox"), TEXT("RtpTimeoutMs"), Hints.lost_rtp_timeout_ms, GEngineIni);
+
+	UE_LOG(LogVivoxVoiceChat, Verbose, TEXT("Rtp timeouts configured to: %ims %ims"), Hints.never_rtp_timeout_ms, Hints.lost_rtp_timeout_ms);
 }
 
 VivoxClientApi::AccountName FVivoxVoiceChat::CreateAccountName(const FString& PlayerName)
