Voice chat error refactor

Upgrade Notes:

Changes to IVoiceChat interface

* VoiceChat.h has moved from "Engine/Source/Runtime/Online/Voice/Public/Interfaces" to a new VoiceChat header-only plugin in "Engine/Plugins/Online/VoiceChat/VoiceChat/Source/Public". This header can be included by adding "VoiceChat" to the PublicIncludePathModuleNames/PrivateIncludePathModuleNames in Build.cs files as necessary.
* FVoiceChatResult::bSuccess has been replaced with FVoiceChatResult::IsSuccess().
* FVoiceChatResult::ResultCode enum added, providing either "Success" or, in error conditions, a broad category for the error.
* FVoiceChatResult::ErrorCode changed from an int to an FString. This avoids collisions, and allows for shared error codes (e.g. errors.com.epicgames.voicechat.not_initialized), and implementation specific ones (e.g. errors.com.epicgames.voicechat.vivox.missing_config). The intent is that errors will be mapped to one of the common error categories in EVoiceChatResult if possible (or EVoiceChatResult::ImplementationError if none apply), and to one of the common error codes in VoiceChatErrors.h. If no common error codes apply, then the implementation can provide its own error codes (see VivoxVoiceChatErrors.h for reference)
* FVoiceChatResult::Error has been renamed to ErrorDesc. The intent is for this to provide additional information for logging.

Changes to VivoxVoiceChat plugin:

* Errors originating from the _plugin_ not the sdk, which were previously reported with a small negative integer and FString (e.g. -2 and "Not Connected"), have been changed to provide error categories and string codes (e.g. EVoiceChatResult::NotConnected and errors.com.epicgames.voicechat.not_connected)
* Numeric errors originating from the vivox sdk itself were previously reported directly in the ErrorCode. These are now mapped in the following ways:
    * Where possible, some errors, e.g. VX_E_NOT_INITIALIZED are mapped to common error categories and codes, in this case EVoiceChatResult::NotInitialized and "errors.com.epicgames.voicechat.not_initialized".
    * Some errors are not mappable to common error codes and are therefore mapped to vivox specific error codes, with the error category EVoiceChatResult::ImplementationError. For example VX_E_CALL_TERMINATED_KICK is mapped to the error code "errors.com.epicgames.voicechat.vivox.kicked_from_channel".
    * Finally any errors that are not mapped to specific error codes are reported as EVoiceChatResult::ImplementationError and "errors.com.epicgames.voicechat.vivox.<numeric_code>".
    * The actual numeric error code and status string sent by the sdk are provided in FVoiceChatResult::ErrorDesc in the format "StatusCode=%d StatusString=[%s]"
    Please look to VivoxVoiceChat.cpp and especially FVivoxVoiceChat::ResultFromVivoxStatus to see the current mapping. If mappings change or are added in the future we will call this out in upgrade notes.

diff --git a/Engine/Plugins/Online/VoiceChat/VivoxVoiceChat/Source/Private/Android/AndroidVivoxVoiceChat.cpp b/Engine/Plugins/Online/VoiceChat/VivoxVoiceChat/Source/Private/Android/AndroidVivoxVoiceChat.cpp
--- a/Engine/Plugins/Online/VoiceChat/VivoxVoiceChat/Source/Private/Android/AndroidVivoxVoiceChat.cpp
+++ b/Engine/Plugins/Online/VoiceChat/VivoxVoiceChat/Source/Private/Android/AndroidVivoxVoiceChat.cpp
@@ -100,7 +100,7 @@ bool FAndroidVivoxVoiceChat::Uninitialize()
 
 void FAndroidVivoxVoiceChat::OnVoiceChatConnectComplete(const FVoiceChatResult& Result)
 {
-	if (Result.bSuccess)
+	if (Result.IsSuccess())
 	{
 		OnVoiceChatReconnectedDelegate.Broadcast();
 	}
diff --git a/Engine/Plugins/Online/VoiceChat/VivoxVoiceChat/Source/Private/IOS/IOSVivoxVoiceChat.cpp b/Engine/Plugins/Online/VoiceChat/VivoxVoiceChat/Source/Private/IOS/IOSVivoxVoiceChat.cpp
--- a/Engine/Plugins/Online/VoiceChat/VivoxVoiceChat/Source/Private/IOS/IOSVivoxVoiceChat.cpp
+++ b/Engine/Plugins/Online/VoiceChat/VivoxVoiceChat/Source/Private/IOS/IOSVivoxVoiceChat.cpp
@@ -127,7 +127,7 @@ void FIOSVivoxVoiceChat::JoinChannel(const FString& ChannelName, const FString&
 	FOnVoiceChatChannelJoinCompleteDelegate DelegateWrapper = FOnVoiceChatChannelJoinCompleteDelegate::CreateLambda(
 		[this, Delegate](const FString& ChannelName, const FVoiceChatResult& Result)
 		{
-			if (!Result.bSuccess)
+			if (!Result.IsSuccess())
 			{
 				EnableVoiceChat(false);
 			}
@@ -188,7 +188,7 @@ void FIOSVivoxVoiceChat::onDisconnected(const VivoxClientApi::Uri& Server, const
 
 void FIOSVivoxVoiceChat::OnVoiceChatConnectComplete(const FVoiceChatResult& Result)
 {
-	if (Result.bSuccess)
+	if (Result.IsSuccess())
 	{
 		OnVoiceChatReconnectedDelegate.Broadcast();
 	}
diff --git a/Engine/Plugins/Online/VoiceChat/VivoxVoiceChat/Source/Private/VivoxVoiceChat.cpp b/Engine/Plugins/Online/VoiceChat/VivoxVoiceChat/Source/Private/VivoxVoiceChat.cpp
--- a/Engine/Plugins/Online/VoiceChat/VivoxVoiceChat/Source/Private/VivoxVoiceChat.cpp
+++ b/Engine/Plugins/Online/VoiceChat/VivoxVoiceChat/Source/Private/VivoxVoiceChat.cpp
@@ -3,12 +3,14 @@
 #include "VivoxVoiceChat.h" 
 
 #include "Async/Async.h"
+#include "HAL/LowLevelMemTracker.h"
 #include "Logging/LogMacros.h"
 #include "Misc/ConfigCacheIni.h"
 #include "Misc/EmbeddedCommunication.h"
 #include "Misc/ScopeLock.h"
 #include "Modules/ModuleManager.h"
-#include "HAL/LowLevelMemTracker.h"
+#include "VoiceChatErrors.h"
+#include "VivoxVoiceChatErrors.h"
 
 #include "Vxc.h"
 #include "VxcErrors.h"
@@ -19,7 +21,6 @@ FVivoxDelegates::FOnAudioUnitCaptureDeviceStatusChanged FVivoxDelegates::OnAudio
 
 namespace
 {
-static const FVoiceChatResult ResultSuccess = { true, 0, TEXT("") };
 
 FString LexToString(VivoxClientApi::IClientApiEventHandler::ParticipantLeftReason Reason)
 {
@@ -33,35 +34,75 @@ FString LexToString(VivoxClientApi::IClientApiEventHandler::ParticipantLeftReaso
 	}
 }
 
-FString LexToString(const FVoiceChatResult& Result)
-{
-	if (Result.bSuccess)
-	{
-		return TEXT("Success");
-	}
-	else if (!Result.Error.IsEmpty())
-	{
-		return FString::Printf(TEXT("Failed: %s"), *Result.Error);
-	}
-	else
-	{
-		return FString::Printf(TEXT("Failed: Error %i"), Result.ErrorCode);
-	}
-}
-
 FVoiceChatResult ResultFromVivoxStatus(const VivoxClientApi::VCSStatus& Status)
 {
-	FString ErrorString = ANSI_TO_TCHAR(Status.ToString());
-	if (ErrorString.IsEmpty() && Status.IsError())
+	FVoiceChatResult Result = FVoiceChatResult::CreateSuccess();
+	if (Status.IsError())
 	{
-		ErrorString = FString::Printf(TEXT("Error %i"), Status.GetStatusCode());
+		switch (Status.GetStatusCode())
+		{
+		case VX_E_NOT_INITIALIZED:
+			Result = VoiceChat::Errors::NotInitialized();
+			break;
+		case VX_E_ALREADY_LOGGED_OUT:
+		case VX_E_NOT_LOGGED_IN:
+			Result = VoiceChat::Errors::NotLoggedIn();
+			break;
+		case VX_E_INVALID_ARGUMENT:
+		case VX_E_INVALID_USERNAME_OR_PASSWORD:
+		case VX_E_CHANNEL_URI_TOO_LONG:
+			Result = VoiceChat::Errors::InvalidArgument();
+			break;
+		case VX_E_CALL_TERMINATED_NO_ANSWER_LOCAL:
+			Result = VoiceChat::Errors::ClientTimeout();
+			break;
+		case VX_E_RTP_TIMEOUT:
+		case VxNetworkHttpTimeout:
+			Result = VoiceChat::Errors::ServerTimeout();
+			break;
+		case VX_E_FAILED_TO_CONNECT_TO_SERVER:
+			Result = VoiceChat::Errors::ConnectionFailure();
+			break;
+		case 10005: // Couldn't resolve proxy name
+		case VxNetworkNameResolutionFailed:
+			Result = VoiceChat::Errors::DnsFailure();
+			break;
+		case VX_E_ACCESSTOKEN_INVALID_SIGNATURE:
+		case VX_E_ACCESSTOKEN_CLAIMS_MISMATCH:
+		case VX_E_ACCESSTOKEN_ISSUER_MISMATCH:
+		case VX_E_ACCESSTOKEN_MALFORMED:
+			Result = VoiceChat::Errors::CredentialsInvalid();
+			break;
+		case VX_E_ACCESSTOKEN_ALREADY_USED:
+		case VX_E_ACCESSTOKEN_EXPIRED:
+			Result = VoiceChat::Errors::CredentialsExpired();
+			break;
+		case VX_E_ALREADY_INITIALIZED:
+			Result = VivoxVoiceChat::Errors::AlreadyInitialized();
+			break;
+		case VX_E_ALREADY_LOGGED_IN:
+			Result = VivoxVoiceChat::Errors::AlreadyLoggedIn();
+			break;
+		case VX_E_CALL_TERMINATED_KICK:
+			Result = VivoxVoiceChat::Errors::KickedFromChannel();
+			break;
+		case VX_E_NO_EXIST:
+			Result = VivoxVoiceChat::Errors::NoExist();
+			break;
+		case VX_E_MAXIMUM_NUMBER_OF_CALLS_EXCEEEDED:
+			Result = VivoxVoiceChat::Errors::MaximumNumberOfCallsExceeded();
+			break;
+		default:
+			// TODO map more vivox statuses to text error codes
+			Result = VIVOXVOICECHAT_ERROR(EVoiceChatResult::ImplementationError, FString::FromInt(Status.GetStatusCode()));
+			break;
+		}
+
+		Result.ErrorDesc = FString::Printf(TEXT("StatusCode=%d StatusString=[%s]"), Status.GetStatusCode(), ANSI_TO_TCHAR(Status.ToString()));
+		Result.ErrorNum = Status.GetStatusCode();
 	}
-	return { !Status.IsError(), Status.GetStatusCode(), ErrorString };
-}
 
-FVoiceChatResult ResultFromErrorString(const FString& Error, int ErrorCode)
-{
-	return { false, ErrorCode, Error };
+	return Result;
 }
 
 template<class TDelegate, class... TArgs>
@@ -258,7 +299,8 @@ bool FVivoxVoiceChat::Initialize()
 				VivoxClientApi::VCSStatus Status = VivoxClientConnection.Initialize(this, VivoxLogLevel, true, false, &ConfigHints, sizeof(ConfigHints));
 				if (Status.IsError())
 				{
-					UE_LOG(LogVivoxVoiceChat, Warning, TEXT("Initialize failed: error:%s (%i)"), ANSI_TO_TCHAR(Status.ToString()), Status.GetStatusCode());
+					const FVoiceChatResult VoiceChatResult = ResultFromVivoxStatus(Status);
+					UE_LOG(LogVivoxVoiceChat, Warning, TEXT("Initialize %s"), *LexToString(VoiceChatResult));
 				}
 				else
 				{
@@ -480,26 +522,26 @@ FString FVivoxVoiceChat::GetDefaultOutputDevice() const
 
 void FVivoxVoiceChat::Connect(const FOnVoiceChatConnectCompleteDelegate& Delegate)
 {
-	FVoiceChatResult Result = ResultSuccess;
+	FVoiceChatResult Result = FVoiceChatResult::CreateSuccess();
 
 	if (!IsInitialized())
 	{
-		Result = ResultFromErrorString(TEXT("Not initialized"), VX_E_NOT_INITIALIZED);
+		Result = VoiceChat::Errors::NotInitialized();
 	}
 	else if (ConnectionState == EConnectionState::Disconnecting)
 	{
-		Result = ResultFromErrorString("Disconnect in progress", -1);
+		Result = VoiceChat::Errors::DisconnectInProgress();
 	}
 
-	if (!Result.bSuccess)
+	if (!Result.IsSuccess())
 	{
-		UE_LOG(LogVivoxVoiceChat, Warning, TEXT("Connect Failed: %s"), *Result.Error);
+		UE_LOG(LogVivoxVoiceChat, Warning, TEXT("Connect %s"), *LexToString(Result));
 
 		Delegate.ExecuteIfBound(Result);
 	}
 	else if (IsConnected())
 	{
-		Delegate.ExecuteIfBound(ResultSuccess);
+		Delegate.ExecuteIfBound(FVoiceChatResult::CreateSuccess());
 	}
 	else
 	{
@@ -510,7 +552,8 @@ void FVivoxVoiceChat::Connect(const FOnVoiceChatConnectCompleteDelegate& Delegat
 			if (VivoxServerUrl.IsEmpty() || VivoxDomain.IsEmpty() || VivoxNamespace.IsEmpty())
 			{
 				UE_LOG(LogVivoxVoiceChat, Warning, TEXT("[VoiceChat.Vivox] ServerUrl, Domain, or Issuer is not set. Vivox voice chat will not work"));
-				Result = ResultFromErrorString(TEXT("Vivox config missing"), -1);
+				Result = VivoxVoiceChat::Errors::MissingConfig();
+				Result.ErrorDesc = FString::Printf(TEXT("ServerUrl=[%s] Domain=[%s] Namespace=[%s]"), *VivoxServerUrl, *VivoxDomain, *VivoxNamespace);
 			}
 			else
 			{
@@ -528,7 +571,7 @@ void FVivoxVoiceChat::Connect(const FOnVoiceChatConnectCompleteDelegate& Delegat
 				}
 			}
 
-			if (!Result.bSuccess)
+			if (!Result.IsSuccess())
 			{
 				ConnectionState = EConnectionState::Disconnected;
 
@@ -550,7 +593,7 @@ void FVivoxVoiceChat::Disconnect(const FOnVoiceChatDisconnectCompleteDelegate& D
 	}
 	else
 	{
-		Delegate.ExecuteIfBound(ResultSuccess);
+		Delegate.ExecuteIfBound(FVoiceChatResult::CreateSuccess());
 	}
 }
 
@@ -566,45 +609,45 @@ bool FVivoxVoiceChat::IsConnected() const
 
 void FVivoxVoiceChat::Login(FPlatformUserId PlatformId, const FString& PlayerName, const FString& Credentials, const FOnVoiceChatLoginCompleteDelegate& Delegate)
 {
-	FVoiceChatResult Result = ResultSuccess;
+	FVoiceChatResult Result = FVoiceChatResult::CreateSuccess();
 
 	if (!IsInitialized())
 	{
-		Result = ResultFromErrorString(TEXT("Not Initialized"), VX_E_NOT_INITIALIZED);
+		Result = VoiceChat::Errors::NotInitialized();
 	}
 	else if (!IsConnected())
 	{
-		Result = ResultFromErrorString(TEXT("Not Connected"), -2);
+		Result = VoiceChat::Errors::NotConnected();
 	}
 	else if (IsLoggedIn())
 	{
 		if (PlayerName == GetLoggedInPlayerName())
 		{
-			Delegate.ExecuteIfBound(PlayerName, ResultSuccess);
+			Delegate.ExecuteIfBound(PlayerName, FVoiceChatResult::CreateSuccess());
 			return;
 		}
 		else
 		{
-			Result = ResultFromErrorString(TEXT("Other user logged in"), VX_E_ALREADY_LOGGED_IN);
+			Result = VoiceChat::Errors::OtherUserLoggedIn();
 		}
 	}
 	else if (PlayerName.IsEmpty())
 	{
-		Result = ResultFromErrorString(TEXT("Player name is empty"), VX_E_INVALID_ARGUMENT);
+		Result = VoiceChat::Errors::InvalidArgument(TEXT("PlayerName empty"));
 	}
 	else if (!VivoxNameContainsValidCharacters(PlayerName))
 	{
-		Result = ResultFromErrorString(TEXT("Invalid PlayerName"), VX_E_INVALID_ARGUMENT);
+		Result = VoiceChat::Errors::InvalidArgument(TEXT("PlayerName invalid characters"));
 	}
-	else if (PlayerName.Len() > 60 - VivoxNamespace.Len())
+	else if(PlayerName.Len() > 60 - VivoxNamespace.Len())
 	{
 		// Name must be between 3-63 characters long and start and end with a '.'. It also must contain the issuer and another '.' separating issuer and name
-		Result = ResultFromErrorString(TEXT("PlayerName is too long"), -1);
+		Result = VoiceChat::Errors::InvalidArgument(TEXT("PlayerName too long"));
 	}
 
-	if (!Result.bSuccess)
+	if (!Result.IsSuccess())
 	{
-		UE_LOG(LogVivoxVoiceChat, Warning, TEXT("Login failed: PlayerName:%s error:%s"), *PlayerName, *Result.Error);
+		UE_LOG(LogVivoxVoiceChat, Warning, TEXT("Login PlayerName:%s %s"), *PlayerName, *LexToString(Result));
 		Delegate.ExecuteIfBound(PlayerName, Result);
 		return;
 	}
@@ -614,9 +657,9 @@ void FVivoxVoiceChat::Login(FPlatformUserId PlatformId, const FString& PlayerNam
 
 	if (!AccountName.IsValid() || !UserUri.IsValid())
 	{
-		VivoxClientApi::VCSStatus Status(VX_E_INVALID_ARGUMENT);
-		UE_LOG(LogVivoxVoiceChat, Warning, TEXT("Login failed: PlayerName:%s error:%s (%i)"), *PlayerName, ANSI_TO_TCHAR(Status.ToString()), Status.GetStatusCode());
-		Delegate.ExecuteIfBound(PlayerName, ResultFromVivoxStatus(Status));
+		Result = VoiceChat::Errors::CredentialsInvalid();
+		UE_LOG(LogVivoxVoiceChat, Warning, TEXT("Login PlayerName:%s %s"), *PlayerName, *LexToString(Result));
+		Delegate.ExecuteIfBound(PlayerName, Result);
 		return;
 	}
 
@@ -629,8 +672,9 @@ void FVivoxVoiceChat::Login(FPlatformUserId PlatformId, const FString& PlayerNam
 	VivoxClientApi::VCSStatus Status = VivoxClientConnection.Login(LoginSession.AccountName, TCHAR_TO_ANSI(*Credentials));
 	if (Status.IsError())
 	{
-		UE_LOG(LogVivoxVoiceChat, Warning, TEXT("Login failed: account:%s error:%s (%i)"), ANSI_TO_TCHAR(LoginSession.AccountName.ToString()), ANSI_TO_TCHAR(Status.ToString()), Status.GetStatusCode());
-		Delegate.ExecuteIfBound(PlayerName, ResultFromVivoxStatus(Status));
+		Result = ResultFromVivoxStatus(Status);
+		UE_LOG(LogVivoxVoiceChat, Warning, TEXT("Login account:%s %s"), ANSI_TO_TCHAR(LoginSession.AccountName.ToString()), *LexToString(Result));
+		Delegate.ExecuteIfBound(PlayerName, Result);
 		return;
 	}
 
@@ -641,25 +685,25 @@ void FVivoxVoiceChat::Login(FPlatformUserId PlatformId, const FString& PlayerNam
 
 void FVivoxVoiceChat::Logout(const FOnVoiceChatLogoutCompleteDelegate& Delegate)
 {
-	FVoiceChatResult Result = ResultSuccess;
+	FVoiceChatResult Result = FVoiceChatResult::CreateSuccess();
 
 	if (!IsInitialized())
 	{
-		Result = ResultFromErrorString(TEXT("Not Initialized"), VX_E_NOT_INITIALIZED);
+		Result = VoiceChat::Errors::NotInitialized();
 	}
 	else if (!IsConnected())
 	{
-		Result = ResultFromErrorString(TEXT("Not Connected"), -2);
+		Result = VoiceChat::Errors::NotConnected();
 	}
 	else if (!IsLoggedIn())
 	{
-		Result = ResultFromErrorString(TEXT("Not Logged In"), VX_E_NOT_LOGGED_IN);
+		Result = VoiceChat::Errors::NotLoggedIn();
 	}
 	// TODO: handle IsLoggingIn case
 
-	if (!Result.bSuccess)
+	if (!Result.IsSuccess())
 	{
-		UE_LOG(LogVivoxVoiceChat, Warning, TEXT("Logout failed: error:%s"), *Result.Error);
+		UE_LOG(LogVivoxVoiceChat, Warning, TEXT("Logout %s"), *LexToString(Result));
 		Delegate.ExecuteIfBound(FString(), Result);
 		return;
 	}
@@ -668,8 +712,9 @@ void FVivoxVoiceChat::Logout(const FOnVoiceChatLogoutCompleteDelegate& Delegate)
 	VivoxClientApi::VCSStatus Status = VivoxClientConnection.Logout(AccountName);
 	if (Status.IsError())
 	{
-		UE_LOG(LogVivoxVoiceChat, Warning, TEXT("Logout failed: error:%s (%i)"), ANSI_TO_TCHAR(Status.ToString()), Status.GetStatusCode());
-		Delegate.ExecuteIfBound(LoginSession.PlayerName, ResultFromVivoxStatus(Status));
+		Result = ResultFromVivoxStatus(Status);
+		UE_LOG(LogVivoxVoiceChat, Warning, TEXT("Logout %s"), *LexToString(Result));
+		Delegate.ExecuteIfBound(LoginSession.PlayerName, Result);
 		return;
 	}
 
@@ -738,37 +783,37 @@ void FVivoxVoiceChat::UnblockPlayers(const TArray<FString>& PlayerNames)
 
 void FVivoxVoiceChat::JoinChannel(const FString& ChannelName, const FString& ChannelCredentials, EVoiceChatChannelType ChannelType, const FOnVoiceChatChannelJoinCompleteDelegate& Delegate, TOptional<FVoiceChatChannel3dProperties> Channel3dProperties)
 {
-	FVoiceChatResult Result = ResultSuccess;
+	FVoiceChatResult Result = FVoiceChatResult::CreateSuccess();
 
 	if (!IsInitialized())
 	{
-		Result = ResultFromErrorString(TEXT("Not Initialized"), VX_E_NOT_INITIALIZED);
+		Result = VoiceChat::Errors::NotInitialized();
 	}
 	else if (!IsConnected())
 	{
-		Result = ResultFromErrorString(TEXT("Not Connected"), -2);
+		Result = VoiceChat::Errors::NotConnected();
 	}
 	else if (!IsLoggedIn())
 	{
-		Result = ResultFromErrorString(TEXT("Not Logged In"), VX_E_NOT_LOGGED_IN);
+		Result = VoiceChat::Errors::NotLoggedIn();
 	}
 	else if (ChannelName.IsEmpty())
 	{
-		Result = ResultFromErrorString(TEXT("ChannelName is empty"), VX_E_INVALID_ARGUMENT);
+		Result = VoiceChat::Errors::InvalidArgument(TEXT("ChannelName empty"));
 	}
 	else if (!VivoxNameContainsValidCharacters(ChannelName))
 	{
-		Result = ResultFromErrorString(TEXT("Invalid ChannelName"), VX_E_INVALID_ARGUMENT);
+		Result = VoiceChat::Errors::InvalidArgument(TEXT("ChannelName invalid characters"));
 	}
 	else if (ChannelName.Len() > 189 - VivoxNamespace.Len())
 	{
 		// channel name length must not exceed 200 characters, including the confctl-?- prefix and the issuer and separator
-		Result = ResultFromErrorString(TEXT("ChannelName is too long"), VX_E_INVALID_ARGUMENT);
+		Result = VoiceChat::Errors::InvalidArgument(TEXT("ChannelName too long"));
 	}
 
-	if (!Result.bSuccess)
+	if (!Result.IsSuccess())
 	{
-		UE_LOG(LogVivoxVoiceChat, Warning, TEXT("JoinChannel failed: ChannelName:%s error:%s"), *ChannelName, *Result.Error);
+		UE_LOG(LogVivoxVoiceChat, Warning, TEXT("JoinChannel ChannelName:%s %s"), *ChannelName, *LexToString(Result));
 		Delegate.ExecuteIfBound(ChannelName, Result);
 		return;
 	}
@@ -776,17 +821,17 @@ void FVivoxVoiceChat::JoinChannel(const FString& ChannelName, const FString& Cha
 	FChannelSession& ChannelSession = GetChannelSession(ChannelName);
 	if (ChannelSession.State == FChannelSession::EState::Connected)
 	{
-		Delegate.ExecuteIfBound(ChannelName, ResultSuccess);
+		Delegate.ExecuteIfBound(ChannelName, FVoiceChatResult::CreateSuccess());
 		return;
 	}
 	else if (ChannelSession.State == FChannelSession::EState::Connecting)
 	{
-		Delegate.ExecuteIfBound(ChannelName, ResultFromErrorString("Join in progress", -1));
+		Delegate.ExecuteIfBound(ChannelName, VoiceChat::Errors::ChannelJoinInProgress());
 		return;
 	}
 	else if (ChannelSession.State == FChannelSession::EState::Disconnecting)
 	{
-		Delegate.ExecuteIfBound(ChannelName, ResultFromErrorString("Leave in progress", -1));
+		Delegate.ExecuteIfBound(ChannelName, VoiceChat::Errors::ChannelLeaveInProgress());
 		return;
 	}
 
@@ -797,8 +842,9 @@ void FVivoxVoiceChat::JoinChannel(const FString& ChannelName, const FString& Cha
 	VivoxClientApi::VCSStatus Status = VivoxClientConnection.JoinChannel(LoginSession.AccountName, ChannelSession.ChannelUri, TCHAR_TO_ANSI(*ChannelCredentials));
 	if (Status.IsError())
 	{
-		UE_LOG(LogVivoxVoiceChat, Warning, TEXT("JoinChannel failed: channel:%s error:%s (%i)"), ANSI_TO_TCHAR(ChannelSession.ChannelUri.ToString()), ANSI_TO_TCHAR(Status.ToString()), Status.GetStatusCode());
-		Delegate.ExecuteIfBound(ChannelName, ResultFromVivoxStatus(Status));
+		Result = ResultFromVivoxStatus(Status);
+		UE_LOG(LogVivoxVoiceChat, Warning, TEXT("JoinChannel ChannelUri:%s %s"), ANSI_TO_TCHAR(ChannelSession.ChannelUri.ToString()), *LexToString(Result));
+		Delegate.ExecuteIfBound(ChannelName, Result);
 		return;
 	}
 
@@ -808,34 +854,34 @@ void FVivoxVoiceChat::JoinChannel(const FString& ChannelName, const FString& Cha
 
 void FVivoxVoiceChat::LeaveChannel(const FString& ChannelName, const FOnVoiceChatChannelLeaveCompleteDelegate& Delegate)
 {
-	FVoiceChatResult Result = ResultSuccess;
+	FVoiceChatResult Result = FVoiceChatResult::CreateSuccess();
 
 	if (!IsInitialized())
 	{
-		Result = ResultFromErrorString(TEXT("Not Initialized"), VX_E_NOT_INITIALIZED);
+		Result = VoiceChat::Errors::NotInitialized();
 	}
 	else if (!IsConnected())
 	{
-		Result = ResultFromErrorString(TEXT("Not Connected"), -2);
+		Result = VoiceChat::Errors::NotConnected();
 	}
 	else if (!IsLoggedIn())
 	{
-		Result = ResultFromErrorString(TEXT("Not Logged In"), VX_E_NOT_LOGGED_IN);
+		Result = VoiceChat::Errors::NotLoggedIn();
 	}
 	else if (ChannelName.IsEmpty())
 	{
-		Result = ResultFromErrorString(TEXT("ChannelName is empty"), VX_E_INVALID_ARGUMENT);
+		Result = VoiceChat::Errors::InvalidArgument();
 	}
 
 	FChannelSession& ChannelSession = GetChannelSession(ChannelName);
 	if (ChannelSession.State != FChannelSession::EState::Connected)
 	{
-		Result = ResultFromErrorString("Not in channel", -1);
+		Result = VoiceChat::Errors::NotInChannel();
 	}
 
-	if (!Result.bSuccess)
+	if (!Result.IsSuccess())
 	{
-		UE_LOG(LogVivoxVoiceChat, Warning, TEXT("LeaveChannel failed: ChannelName:%s error:%s"), *ChannelName, *Result.Error);
+		UE_LOG(LogVivoxVoiceChat, Warning, TEXT("LeaveChannel ChannelName:%s %s"), *ChannelName, *LexToString(Result));
 		Delegate.ExecuteIfBound(ChannelName, Result);
 		return;
 	}
@@ -843,8 +889,9 @@ void FVivoxVoiceChat::LeaveChannel(const FString& ChannelName, const FOnVoiceCha
 	VivoxClientApi::VCSStatus Status = VivoxClientConnection.LeaveChannel(LoginSession.AccountName, ChannelSession.ChannelUri);
 	if (Status.IsError())
 	{
-		UE_LOG(LogVivoxVoiceChat, Warning, TEXT("LeaveChannel failed: channel:%s error:%s (%i)"), ANSI_TO_TCHAR(ChannelSession.ChannelUri.ToString()), ANSI_TO_TCHAR(Status.ToString()), Status.GetStatusCode());
-		Delegate.ExecuteIfBound(ChannelName, ResultFromVivoxStatus(Status));
+		Result = ResultFromVivoxStatus(Status);
+		UE_LOG(LogVivoxVoiceChat, Warning, TEXT("LeaveChannel channel:%s %s"), ANSI_TO_TCHAR(ChannelSession.ChannelUri.ToString()), *LexToString(Result));
+		Delegate.ExecuteIfBound(ChannelName, Result);
 		return;
 	}
 
@@ -1151,17 +1198,15 @@ void FVivoxVoiceChat::InvokeOnUIThread(void (Func)(void* Arg0), void* Arg0)
 
 FString GetLogLevelString(VivoxClientApi::IClientApiEventHandler::LogLevel Level)
 {
-	FString LogLevelString;
 	switch (Level)
 	{
-	case VivoxClientApi::IClientApiEventHandler::LogLevel::LogLevelError:		LogLevelString = TEXT("Error"); break;
-	case VivoxClientApi::IClientApiEventHandler::LogLevel::LogLevelWarning:		LogLevelString = TEXT("Warning"); break;
-	case VivoxClientApi::IClientApiEventHandler::LogLevel::LogLevelInfo:		LogLevelString = TEXT("Info"); break;
-	case VivoxClientApi::IClientApiEventHandler::LogLevel::LogLevelDebug:		LogLevelString = TEXT("Debug"); break;
-	case VivoxClientApi::IClientApiEventHandler::LogLevel::LogLevelTrace:		LogLevelString = TEXT("Trace"); break;
-	default:																	LogLevelString = TEXT("Unknown"); break;
+	case VivoxClientApi::IClientApiEventHandler::LogLevel::LogLevelError:		return TEXT("Error");
+	case VivoxClientApi::IClientApiEventHandler::LogLevel::LogLevelWarning:		return TEXT("Warning");
+	case VivoxClientApi::IClientApiEventHandler::LogLevel::LogLevelInfo:		return TEXT("Info");
+	case VivoxClientApi::IClientApiEventHandler::LogLevel::LogLevelDebug:		return TEXT("Debug");
+	case VivoxClientApi::IClientApiEventHandler::LogLevel::LogLevelTrace:		return TEXT("Trace");
+	default:																	return TEXT("Unknown");
 	}
-	return LogLevelString;
 }
 
 void FVivoxVoiceChat::onLogStatementEmitted(LogLevel Level, long long NativeMillisecondsSinceEpoch, long ThreadId, const char* LogMessageCStr)
@@ -1198,7 +1243,7 @@ void FVivoxVoiceChat::onLogStatementEmitted(LogLevel Level, long long NativeMill
 	{
 		if (Level == LogLevelError)
 		{
-			UE_LOG(LogVivoxVoiceChat, Log, TEXT("vivox: Error: %s"), *LogMessage);
+			UE_LOG(LogVivoxVoiceChat, Warning, TEXT("vivox: Error: %s"), *LogMessage);
 		}
 		else
 		{
@@ -1217,7 +1262,7 @@ void FVivoxVoiceChat::onConnectCompleted(const VivoxClientApi::Uri& Server)
 
 	ConnectionState = EConnectionState::Connected;
 
-	TriggerCompletionDelegates(OnVoiceChatConnectCompleteDelegates, ResultSuccess);
+	TriggerCompletionDelegates(OnVoiceChatConnectCompleteDelegates, FVoiceChatResult::CreateSuccess());
 }
 
 void FVivoxVoiceChat::onConnectFailed(const VivoxClientApi::Uri& Server, const VivoxClientApi::VCSStatus& Status)
@@ -1247,7 +1292,7 @@ void FVivoxVoiceChat::onDisconnected(const VivoxClientApi::Uri& Server, const Vi
 
 	if (PreviousConnectionState == EConnectionState::Disconnecting)
 	{
-		TriggerCompletionDelegates(OnVoiceChatDisconnectCompleteDelegates, ResultSuccess);
+		TriggerCompletionDelegates(OnVoiceChatDisconnectCompleteDelegates, FVoiceChatResult::CreateSuccess());
 	}
 	else
 	{
@@ -1263,7 +1308,7 @@ void FVivoxVoiceChat::onLoginCompleted(const VivoxClientApi::AccountName& Accoun
 
 	LoginSession.State = FLoginSession::EState::LoggedIn;
 
-	TriggerCompletionDelegate(OnVoiceChatLoginCompleteDelegate, PlayerName, ResultSuccess);
+	TriggerCompletionDelegate(OnVoiceChatLoginCompleteDelegate, PlayerName, FVoiceChatResult::CreateSuccess());
 
 	OnVoiceChatLoggedInDelegate.Broadcast(PlayerName);
 }
@@ -1276,7 +1321,7 @@ void FVivoxVoiceChat::onInvalidLoginCredentials(const VivoxClientApi::AccountNam
 
 	FString PlayerName = GetPlayerNameFromAccountName(AccountName);
 
-	TriggerCompletionDelegate(OnVoiceChatLoginCompleteDelegate, PlayerName, ResultFromErrorString(TEXT("Invalid login credentials"), -1));
+	TriggerCompletionDelegate(OnVoiceChatLoginCompleteDelegate, PlayerName, VoiceChat::Errors::CredentialsInvalid());
 }
 
 void FVivoxVoiceChat::onLoginFailed(const VivoxClientApi::AccountName& AccountName, const VivoxClientApi::VCSStatus& Status)
@@ -1299,7 +1344,7 @@ void FVivoxVoiceChat::onLogoutCompleted(const VivoxClientApi::AccountName& Accou
 	LoginSession.State = FLoginSession::EState::LoggedOut;
 
 	FString PlayerName = GetPlayerNameFromAccountName(AccountName);
-	TriggerCompletionDelegate(OnVoiceChatLogoutCompleteDelegate, PlayerName, ResultSuccess);
+	TriggerCompletionDelegate(OnVoiceChatLogoutCompleteDelegate, PlayerName, FVoiceChatResult::CreateSuccess());
 
 	OnVoiceChatLoggedOutDelegate.Broadcast(PlayerName);
 }
@@ -1320,7 +1365,7 @@ void FVivoxVoiceChat::onChannelJoined(const VivoxClientApi::AccountName& Account
 
 	FChannelSession& ChannelSession = GetChannelSession(ChannelUri);
 	ChannelSession.State = FChannelSession::EState::Connected;
-	TriggerCompletionDelegate(ChannelSession.JoinDelegate, ChannelSession.ChannelName, ResultSuccess);
+	TriggerCompletionDelegate(ChannelSession.JoinDelegate, ChannelSession.ChannelName, FVoiceChatResult::CreateSuccess());
 
 	OnVoiceChatChannelJoinedDelegate.Broadcast(ChannelSession.ChannelName);
 }
@@ -1331,7 +1376,7 @@ void FVivoxVoiceChat::onInvalidChannelCredentials(const VivoxClientApi::AccountN
 
 	FChannelSession& ChannelSession = GetChannelSession(ChannelUri);
 	ChannelSession.State = FChannelSession::EState::Disconnected;
-	TriggerCompletionDelegate(ChannelSession.JoinDelegate, ChannelSession.ChannelName, ResultFromErrorString(TEXT("Invalid join credentials"), -1));
+	TriggerCompletionDelegate(ChannelSession.JoinDelegate, ChannelSession.ChannelName, VoiceChat::Errors::CredentialsInvalid());
 
 	RemoveChannelSession(ChannelSession.ChannelName);
 }
@@ -1347,15 +1392,16 @@ void FVivoxVoiceChat::onChannelJoinFailed(const VivoxClientApi::AccountName& Acc
 	RemoveChannelSession(ChannelSession.ChannelName);
 }
 
-void FVivoxVoiceChat::onChannelExited(const VivoxClientApi::AccountName& AccountName, const VivoxClientApi::Uri& ChannelUri, const VivoxClientApi::VCSStatus& ReasonCode)
+void FVivoxVoiceChat::onChannelExited(const VivoxClientApi::AccountName& AccountName, const VivoxClientApi::Uri& ChannelUri, const VivoxClientApi::VCSStatus& Status)
 {
-	if (ReasonCode.IsError())
+	const FVoiceChatResult Result = ResultFromVivoxStatus(Status);
+	if (Result.IsSuccess())
 	{
-		UE_LOG(LogVivoxVoiceChat, Warning, TEXT("onChannelExited channel:%s error:%s (%i)"), ANSI_TO_TCHAR(ChannelUri.ToString()), ANSI_TO_TCHAR(ReasonCode.ToString()), ReasonCode.GetStatusCode());
+		UE_LOG(LogVivoxVoiceChat, Log, TEXT("onChannelExited ChannelUri:%s"), ANSI_TO_TCHAR(ChannelUri.ToString()));
 	}
 	else
 	{
-		UE_LOG(LogVivoxVoiceChat, Log, TEXT("onChannelExited channel:%s"), ANSI_TO_TCHAR(ChannelUri.ToString()));
+		UE_LOG(LogVivoxVoiceChat, Warning, TEXT("onChannelExited ChannelUri:%s %s"), ANSI_TO_TCHAR(ChannelUri.ToString()), *LexToString(Result));
 	}
 
 	FChannelSession& ChannelSession = GetChannelSession(ChannelUri);
@@ -1366,16 +1412,16 @@ void FVivoxVoiceChat::onChannelExited(const VivoxClientApi::AccountName& Account
 	if (bWasConnecting)
 	{
 		// timeouts while connecting call onChannelExited instead of OnChannelJoinFailed
-		TriggerCompletionDelegate(ChannelSession.JoinDelegate, ChannelSession.ChannelName, ResultFromVivoxStatus(ReasonCode));
+		TriggerCompletionDelegate(ChannelSession.JoinDelegate, ChannelSession.ChannelName, Result);
 	}
 	else
 	{
 		if (bWasDisconnecting)
 		{
-			TriggerCompletionDelegate(ChannelSession.LeaveDelegate, ChannelSession.ChannelName, ResultFromVivoxStatus(ReasonCode));
+			TriggerCompletionDelegate(ChannelSession.LeaveDelegate, ChannelSession.ChannelName, Result);
 		}
 
-		OnVoiceChatChannelExitedDelegate.Broadcast(ChannelSession.ChannelName, ResultFromVivoxStatus(ReasonCode));
+		OnVoiceChatChannelExitedDelegate.Broadcast(ChannelSession.ChannelName, Result);
 	}
 
 	RemoveChannelSession(ChannelSession.ChannelName);
diff --git a/Engine/Plugins/Online/VoiceChat/VivoxVoiceChat/Source/Public/VivoxVoiceChat.h b/Engine/Plugins/Online/VoiceChat/VivoxVoiceChat/Source/Public/VivoxVoiceChat.h
--- a/Engine/Plugins/Online/VoiceChat/VivoxVoiceChat/Source/Public/VivoxVoiceChat.h
+++ b/Engine/Plugins/Online/VoiceChat/VivoxVoiceChat/Source/Public/VivoxVoiceChat.h
@@ -14,7 +14,7 @@ THIRD_PARTY_INCLUDES_START
 #include "vivoxclientapi/clientconnection.h"
 THIRD_PARTY_INCLUDES_END
 
-#include "Interfaces/VoiceChat.h"
+#include "VoiceChat.h"
 #include "Stats/Stats.h"
 #include "Misc/CoreMisc.h"
 #include "Logging/LogMacros.h"
diff --git a/Engine/Plugins/Online/VoiceChat/VivoxVoiceChat/Source/Public/VivoxVoiceChatErrors.h b/Engine/Plugins/Online/VoiceChat/VivoxVoiceChat/Source/Public/VivoxVoiceChatErrors.h
new file mode 100644
--- /dev/null
+++ b/Engine/Plugins/Online/VoiceChat/VivoxVoiceChat/Source/Public/VivoxVoiceChatErrors.h
@@ -0,0 +1,17 @@
+// Copyright 1998-2019 Epic Games, Inc. All Rights Reserved.
+#pragma once
+
+#define VIVOXVOICECHAT_ERROR(...) FVoiceChatResult::CreateError(TEXT("errors.com.epicgames.voicechat.vivox"), __VA_ARGS__)
+
+namespace VivoxVoiceChat
+{
+	namespace Errors
+	{
+		inline FVoiceChatResult MissingConfig() { return VIVOXVOICECHAT_ERROR(EVoiceChatResult::InvalidState, TEXT("missing_config")); }
+		inline FVoiceChatResult AlreadyInitialized() { return VIVOXVOICECHAT_ERROR(EVoiceChatResult::InvalidState, TEXT("already_initialized")); }
+		inline FVoiceChatResult AlreadyLoggedIn() { return VIVOXVOICECHAT_ERROR(EVoiceChatResult::InvalidState, TEXT("already_logged_in")); }
+		inline FVoiceChatResult KickedFromChannel() { return VIVOXVOICECHAT_ERROR(EVoiceChatResult::ImplementationError, TEXT("kicked_from_channel")); }
+		inline FVoiceChatResult NoExist() { return VIVOXVOICECHAT_ERROR(EVoiceChatResult::ImplementationError, TEXT("no_exist")); }
+		inline FVoiceChatResult MaximumNumberOfCallsExceeded() { return VIVOXVOICECHAT_ERROR(EVoiceChatResult::ImplementationError, TEXT("maximum_number_of_calls_exceeded")); }
+	}
+}
\ No newline at end of file
diff --git a/Engine/Plugins/Online/VoiceChat/VivoxVoiceChat/Source/VivoxVoiceChat.Build.cs b/Engine/Plugins/Online/VoiceChat/VivoxVoiceChat/Source/VivoxVoiceChat.Build.cs
--- a/Engine/Plugins/Online/VoiceChat/VivoxVoiceChat/Source/VivoxVoiceChat.Build.cs
+++ b/Engine/Plugins/Online/VoiceChat/VivoxVoiceChat/Source/VivoxVoiceChat.Build.cs
@@ -21,7 +21,7 @@ namespace UnrealBuildTool.Rules
 			PublicIncludePathModuleNames.AddRange(
 				new string[]
 				{
-					"Voice"
+					"VoiceChat"
 				}
 			);
 
diff --git a/Engine/Plugins/Online/VoiceChat/VoiceChat/Source/Public/VoiceChat.h b/Engine/Plugins/Online/VoiceChat/VoiceChat/Source/Public/VoiceChat.h
new file mode 100644
--- /dev/null
+++ b/Engine/Plugins/Online/VoiceChat/VoiceChat/Source/Public/VoiceChat.h
@@ -0,0 +1,566 @@
+// Copyright 1998-2019 Epic Games, Inc. All Rights Reserved.
+
+#pragma once
+
+#include "Containers/ArrayView.h"
+#include "Delegates/Delegate.h"
+#include "Features/IModularFeature.h"
+#include "Features/IModularFeatures.h"
+#include "VoiceChatResult.h"
+
+enum class EVoiceChatChannelType
+{
+	/** Non positional/2d audio channel */
+	NonPositional,
+	/** Positional/3d audio channel */
+	Positional,
+	/** Echo channel. Will only ever have one player and will echo anything you say */
+	Echo
+};
+
+enum class EVoiceChatAttenuationModel
+{
+	/** No attenuation is applied. The audio will drop to 0 at MaxDistance */
+	None,
+	/** The attenuation increases in inverse proportion to the distance. The Rolloff is the inverse of the slope of the attenuation curve. */
+	InverseByDistance,
+	/** The attenuation increases in linear proportion to the distance. The Rolloff is the negative slope of the attenuation curve. */
+	LinearByDistance,
+	/** The attenuation increases in inverse proportion to the distance raised to the power of the Rolloff. */
+	ExponentialByDistance
+};
+
+struct FVoiceChatChannel3dProperties
+{
+	/** The model used to determine how loud audio is at different distances */
+	EVoiceChatAttenuationModel AttenuationModel;
+	/** The distance at which the sound will start to attenuate */
+	float MinDistance;
+	/** The distance at which sound will no longer be audible */
+	float MaxDistance;
+	/** How fast the sound attenuates with distance */
+	float Rolloff;
+};
+
+enum class EVoiceChatTransmitMode
+{
+	/** Transmit to none of the channels you are in */
+	None,
+	/** Transmit to all of the channels you are in */
+	All,
+	/** Transmit to a specific channel */
+	Channel
+};
+
+struct FVoiceChatCallStats
+{
+	/** Call length in seconds */
+	double CallLength;
+	/** Minimum measured network latency in seconds. Zero if no measurements made */
+	double LatencyMinMeasuredSeconds;
+	/** Maximum measured network latency in seconds. Zero if no measurements made */
+	double LatencyMaxMeasuredSeconds;
+	/** Average measured network latency in seconds. Zero if no measurements made */
+	double LatencyAverageMeasuredSeconds;
+	/** Total number of packets, both received and lost */
+	int PacketsNumTotal;
+	/** Total number of lost packets */
+	int PacketsNumLost;
+};
+
+DECLARE_DELEGATE_OneParam(FOnVoiceChatConnectCompleteDelegate, const FVoiceChatResult& /* Result */);
+DECLARE_DELEGATE_OneParam(FOnVoiceChatDisconnectCompleteDelegate, const FVoiceChatResult& /* Result */);
+DECLARE_DELEGATE_TwoParams(FOnVoiceChatLoginCompleteDelegate, const FString& /* PlayerName */, const FVoiceChatResult& /* Result */);
+DECLARE_DELEGATE_TwoParams(FOnVoiceChatLogoutCompleteDelegate, const FString& /* PlayerName */, const FVoiceChatResult& /* Result */);
+DECLARE_DELEGATE_TwoParams(FOnVoiceChatChannelJoinCompleteDelegate, const FString& /* ChannelName */, const FVoiceChatResult& /* Result */);
+DECLARE_DELEGATE_TwoParams(FOnVoiceChatChannelLeaveCompleteDelegate, const FString& /* ChannelName */, const FVoiceChatResult& /* Result */);
+
+DECLARE_MULTICAST_DELEGATE(FOnVoiceChatAvailableAudioDevicesChangedDelegate);
+DECLARE_MULTICAST_DELEGATE(FOnVoiceChatReconnectedDelegate);
+DECLARE_MULTICAST_DELEGATE(FOnVoiceChatConnectedDelegate);
+DECLARE_MULTICAST_DELEGATE_OneParam(FOnVoiceChatDisconnectedDelegate, const FVoiceChatResult& /* Reason */);
+DECLARE_MULTICAST_DELEGATE_OneParam(FOnVoiceChatLoggedInDelegate, const FString& /* PlayerName */);
+DECLARE_MULTICAST_DELEGATE_OneParam(FOnVoiceChatLoggedOutDelegate, const FString& /* PlayerName */);
+DECLARE_MULTICAST_DELEGATE_OneParam(FOnVoiceChatCallStatsUpdatedDelegate, const FVoiceChatCallStats& /* CallStats */);
+DECLARE_MULTICAST_DELEGATE_OneParam(FOnVoiceChatChannelJoinedDelegate, const FString& /* ChannelName */);
+DECLARE_MULTICAST_DELEGATE_TwoParams(FOnVoiceChatChannelExitedDelegate, const FString& /* ChannelName */, const FVoiceChatResult& /* Reason */);
+DECLARE_MULTICAST_DELEGATE_TwoParams(FOnVoiceChatPlayerAddedDelegate, const FString& /* ChannelName */, const FString& /* PlayerName */);
+DECLARE_MULTICAST_DELEGATE_TwoParams(FOnVoiceChatPlayerRemovedDelegate, const FString& /* ChannelName */, const FString& /* PlayerName */);
+DECLARE_MULTICAST_DELEGATE_ThreeParams(FOnVoiceChatPlayerTalkingUpdatedDelegate, const FString& /* ChannelName */, const FString& /* PlayerName */, bool /* bIsTalking */);
+DECLARE_MULTICAST_DELEGATE_ThreeParams(FOnVoiceChatPlayerMuteUpdatedDelegate, const FString& /* ChannelName */, const FString& /* PlayerName */, bool /* bIsMuted */);
+DECLARE_MULTICAST_DELEGATE_ThreeParams(FOnVoiceChatPlayerVolumeUpdatedDelegate, const FString& /* ChannelName */, const FString& /* PlayerName */, float /* Volume */);
+
+DECLARE_MULTICAST_DELEGATE_ThreeParams(FOnVoiceChatRecordSamplesAvailableDelegate, TArrayView<const int16> /* PcmSamples */, int /* SampleRate */, int /* Channels */);
+DECLARE_MULTICAST_DELEGATE_ThreeParams(FOnVoiceChatAfterCaptureAudioReadDelegate, TArrayView<int16> /* PcmSamples */, int /* SampleRate */, int /* Channels */);
+DECLARE_MULTICAST_DELEGATE_FourParams(FOnVoiceChatBeforeCaptureAudioSentDelegate, TArrayView<const int16> /* PcmSamples */, int /* SampleRate */, int /* Channels */, bool /* bIsSpeaking */);
+DECLARE_MULTICAST_DELEGATE_FourParams(FOnVoiceChatBeforeRecvAudioRenderedDelegate, TArrayView<int16> /* PcmSamples */, int /* SampleRate */, int /* Channels */, bool /* bIsSilence */);
+
+class IVoiceChat : public IModularFeature
+{
+public:
+	/**
+	 * Get the IVoiceChat modular feature instance
+	 */
+	static IVoiceChat* Get()
+	{
+		if (IModularFeatures::Get().IsModularFeatureAvailable(GetModularFeatureName()))
+		{
+			return &IModularFeatures::Get().GetModularFeature<IVoiceChat>(GetModularFeatureName());
+		}
+		return nullptr;
+	}
+
+	/**
+	 * Initialize VoiceChat
+	 */
+	virtual bool Initialize() = 0;
+
+	/**
+	 * Uninitialize VoiceChat
+	 */
+	virtual bool Uninitialize() = 0;
+
+	/**
+	 * Is voice chat initialized?
+	 *
+	 * @return true if voice chat is initialized
+	 */
+	virtual bool IsInitialized() const = 0;
+
+	/**
+	 * Set an implementation specific setting
+	 *
+	 * @param Name Name of the config setting to set
+	 * @param Value Value for the setting
+	 */
+	virtual void SetSetting(const FString& Name, const FString& Value) = 0;
+
+	/**
+	 * Get an implementation specific setting
+	 *
+	 * @param Name Name of the config setting to get
+	 *
+	 * @return Current value for the setting
+	 */
+	virtual FString GetSetting(const FString& Name) = 0;
+
+	/**
+	 * Set record volume
+	 *
+	 * @param Volume Value between 0.0 and 1.0
+	 */
+	virtual void SetAudioInputVolume(float Volume) = 0;
+
+	/**
+	 * Set playback volume
+	 *
+	 * @param Volume Value between 0.0 and 1.0
+	 */
+	virtual void SetAudioOutputVolume(float Volume) = 0;
+
+	/**
+	 * Get volume of audio input device
+	 *
+	 * @return Audio input device volume. Will be between 0.0 and 1.0
+	 */
+	virtual float GetAudioInputVolume() const = 0;
+
+	/**
+	 * Get volume of audio output device
+	 *
+	 * @return Audio output device volume. Will be between 0.0 and 1.0
+	 */
+	virtual float GetAudioOutputVolume() const = 0;
+
+	/**
+	 * Mute or unmute the audio input device
+	 *
+	 * @param bIsMuted set to true to mute the device or false to unmute
+	 */
+	virtual void SetAudioInputDeviceMuted(bool bIsMuted) = 0;
+
+	/**
+	 * Mute or unmute the audio output device
+	 *
+	 * @param bIsMuted set to true to mute the device or false to unmute
+	 */
+	virtual void SetAudioOutputDeviceMuted(bool bIsMuted) = 0;
+
+	/**
+	 * Is the input device muted?
+	 *
+	 * @return true if the input device is muted
+	 */
+	virtual bool GetAudioInputDeviceMuted() const = 0;
+
+	/**
+	 * Is the output device muted?
+	 *
+	 * @return true if the input device is muted
+	 */
+	virtual bool GetAudioOutputDeviceMuted() const = 0;
+
+	/**
+	 * Get a list of available audio input devices
+	 *
+	 * @return Array of audio input devices
+	 */
+	virtual TArray<FString> GetAvailableInputDevices() const = 0;
+
+	/**
+	 * Get a list of available audio output devices
+	 *
+	 * @return Array of audio output devices
+	 */
+	virtual TArray<FString> GetAvailableOutputDevices() const = 0;
+
+	/**
+	 * Delegate broadcast whenever the available audio devices change. Any cached values from GetAvailableInputDevices or GetAvailableOutputDevices should be discarded and requeried
+	 */
+	virtual FOnVoiceChatAvailableAudioDevicesChangedDelegate& OnVoiceChatAvailableAudioDevicesChanged() = 0;
+
+	/**
+	 * Set the audio input device to use
+	 *
+	 * @param InputDevice String from GetAvailableInputDevice. If empty, will use the default input device
+	 */
+	virtual void SetInputDevice(const FString& InputDevice) = 0;
+
+	/**
+	 * Set the audio output device to use
+	 *
+	 * @param OutputDevice String from GetAvailableOutputDevice. If empty, will use the default output device
+	 */
+	virtual void SetOutputDevice(const FString& OutputDevice) = 0;
+
+	/**
+	 * Get the audio input device that is being used
+	 */
+	virtual FString GetInputDevice() const = 0;
+
+	/**
+	 * Get the audio output device that is being used
+	 */
+	virtual FString GetOutputDevice() const = 0;
+
+	/**
+	 * Get the default audio input device
+	 */
+	virtual FString GetDefaultInputDevice() const = 0;
+
+	/**
+	 * Get the default audio output device
+	 */
+	virtual FString GetDefaultOutputDevice() const = 0;
+
+	/**
+	 * Connect to a voice server
+	 *
+	 * @param Delegate delegate called once connect completes
+	 */
+	virtual void Connect(const FOnVoiceChatConnectCompleteDelegate& Delegate) = 0;
+
+	/**
+	 * Disconnect from a voice server
+	 *
+	 * @param Delegate delegate called once disconnect completes
+	 */
+	virtual void Disconnect(const FOnVoiceChatDisconnectCompleteDelegate& Delegate) = 0;
+
+	/**
+	* Are we connecting to the voice server?
+	*
+	* @return true if we are connecting to the voice server
+	*/
+	virtual bool IsConnecting() const = 0;
+
+	/**
+	 * Are we connected to the voice server?
+	 *
+	 * @return true if we are connected to the voice server
+	 */
+	virtual bool IsConnected() const = 0;
+
+	/**
+	 * Delegate triggered when we are connected to voice chat
+	 */
+	virtual FOnVoiceChatConnectedDelegate& OnVoiceChatConnected() = 0;
+
+	/**
+	 * Delegate triggered when we are disconnected from voice chat
+	 */
+	virtual FOnVoiceChatDisconnectedDelegate& OnVoiceChatDisconnected() = 0;
+
+	/**
+	 * Delegate triggered when we are reconnected to voice chat in cases where the underlying implementation disconnected
+	 * and then reconnected. This is most commonly seen when an application resumes after being suspended. If the disconnect
+	 * is unexpected, OnVoiceChatDisconnected will be called instead.
+	 */
+	virtual FOnVoiceChatReconnectedDelegate& OnVoiceChatReconnected() = 0;
+
+	/**
+	 * Login to the connected voice server
+	 *
+	 * @param PlatformId Platform user Id for the player logging in
+	 * @param PlayerName Name of the account to login to
+	 * @param Credentials Token or password
+	 * @param Delegate delegate called once login completes
+	 */
+	virtual void Login(FPlatformUserId PlatformId, const FString& PlayerName, const FString& Credentials, const FOnVoiceChatLoginCompleteDelegate& Delegate) = 0;
+
+	/**
+	 * Logout from the connected voice server
+	 *
+	 * @param Delegate delegate called once logout completes
+	 */
+	virtual void Logout(const FOnVoiceChatLogoutCompleteDelegate& Delegate) = 0;
+
+	/**
+	 * Are we logging in?
+	 *
+	 * @return true if we are logging in to the voice server
+	 */
+	virtual bool IsLoggingIn() const = 0;
+
+	/**
+	 * Are we logged in?
+	 *
+	 * @return true if we are logged in to the voice server
+	 */
+	virtual bool IsLoggedIn() const = 0;
+
+	/**
+	 * Delegate triggered when we are logged in to voice chat
+	 */
+	virtual FOnVoiceChatLoggedInDelegate& OnVoiceChatLoggedIn() = 0;
+
+	/**
+	 * Delegate triggered when we are logged out from voice chat
+	 */
+	virtual FOnVoiceChatLoggedOutDelegate& OnVoiceChatLoggedOut() = 0;
+
+	/**
+	 * Get the player name used to log in
+	 *
+	 * @return player name used to log in
+	 */
+	virtual FString GetLoggedInPlayerName() const = 0;
+
+	/**
+	 * Block communication with players
+	 *
+	 * Adds the PlayerNames to the block list, muting them in all current and subsequently joined channels.
+	 * A block for a player persists until either:
+	 *		- The local user logs out (e.g. by calling Logout/Disconnect/Uninitialize, or receiving equivalent events)
+	 *		- A subsequent call to UnblockPlayers is made with that PlayerName.
+	 *
+	 * @param PlayerNames List of players to block communication with
+	 */
+	virtual void BlockPlayers(const TArray<FString>& PlayerNames) = 0;
+
+	/**
+	 * Unblock communication with players
+	 *
+	 * @param PlayerNames List of players to unblock communication with
+	 */
+	virtual void UnblockPlayers(const TArray<FString>& PlayerNames) = 0;
+
+	/**
+	 * Join a voice channel
+	 *
+	 * @param ChannelName Channel to join
+	 * @param ChannelCredentials password or token
+	 * @param ChannelType NonPositional, Positional, or Echo
+	 * @param Delegate delegate called once join completes
+	 * @param Channel3dProperties Optional parameter to set 3d/Positional channel properties
+	 */
+	virtual void JoinChannel(const FString& ChannelName, const FString& ChannelCredentials, EVoiceChatChannelType ChannelType, const FOnVoiceChatChannelJoinCompleteDelegate& Delegate, TOptional<FVoiceChatChannel3dProperties> Channel3dProperties = TOptional<FVoiceChatChannel3dProperties>()) = 0;
+
+	/**
+	 * Leave a voice channel
+	 * 
+	 * @param ChannelName Channel to leave
+	 * @param Delegate delegate called once leave completes
+	 */
+	virtual void LeaveChannel(const FString& ChannelName, const FOnVoiceChatChannelLeaveCompleteDelegate& Delegate) = 0;
+
+	/**
+	 * Delegate triggered when we join a voice channel
+	 */
+	virtual FOnVoiceChatChannelJoinedDelegate& OnVoiceChatChannelJoined() = 0;
+
+	/**
+	 * Delegate triggered when we leave a voice channel
+	 */
+	virtual FOnVoiceChatChannelExitedDelegate& OnVoiceChatChannelExited() = 0;
+
+	/**
+	 * Delegate triggered when a call is ended, providing the stats for the call.
+	 */
+	virtual FOnVoiceChatCallStatsUpdatedDelegate& OnVoiceChatCallStatsUpdated() = 0;
+
+	/**
+	 * Set the 3d position of the player
+	 *
+	 * @param ChannelName
+	 * @param SpeakerPosition
+	 * @param ListenerPosition
+	 * @param ListenerForwardDirection
+	 * @param ListenerUpDirection
+	 */
+	virtual void Set3DPosition(const FString& ChannelName, const FVector& SpeakerPosition, const FVector& ListenerPosition, const FVector& ListenerForwardDirection, const FVector& ListenerUpDirection) = 0;
+
+	/**
+	 * Get an array of channels the user is in
+	 *
+	 * @return Array of connected channel names
+	 */
+	virtual TArray<FString> GetChannels() const = 0;
+
+	/**
+	 * Get an array of players in a channel
+	 *
+	 * @param ChannelName Channel to get the list of players for
+	 * @return Array of names of players in Channel
+	 */
+	virtual TArray<FString> GetPlayersInChannel(const FString& ChannelName) const = 0;
+
+	/**
+	 * Get an array of players in a channel
+	 *
+	 * @param ChannelName Channel to get the type of
+	 * @return Channel type
+	 */
+	virtual EVoiceChatChannelType GetChannelType(const FString& ChannelName) const = 0;
+
+	/**
+	 * Delegate triggered when a player is added to a channel
+	 */
+	virtual FOnVoiceChatPlayerAddedDelegate& OnVoiceChatPlayerAdded() = 0;
+
+	/**
+	 * Delegate triggered when a player is removed from a channel
+	 */
+	virtual FOnVoiceChatPlayerRemovedDelegate& OnVoiceChatPlayerRemoved() = 0;
+
+	/**
+	 * @param PlayerName Player to get the talking state of
+	 * @return true if player is talking
+	 */
+	virtual bool IsPlayerTalking(const FString& PlayerName) const = 0;
+
+	/**
+	 * Delegate triggered when a player's talking state is updated
+	 */
+	virtual FOnVoiceChatPlayerTalkingUpdatedDelegate& OnVoiceChatPlayerTalkingUpdated() = 0;
+
+	/**
+	 * Mute or unmute a player
+	 *
+	 * @param PlayerName Player to mute
+	 * @param bMuted true if the player should be muted
+	 */
+	virtual void SetPlayerMuted(const FString& PlayerName, bool bMuted) = 0;
+
+	/**
+	 * @param PlayerName Player to get the mute state of
+	 * @return true if player is muted
+	 */
+	virtual bool IsPlayerMuted(const FString& PlayerName) const = 0;
+
+	/**
+	 * Delegate triggered when a player's muted state is updated, usually as a result of calling SetPlayerMuted
+	 */
+	virtual FOnVoiceChatPlayerMuteUpdatedDelegate& OnVoiceChatPlayerMuteUpdated() = 0;
+
+	/**
+	 * Set player volume
+	 *
+	 * @param PlayerName Player to adjust the volume for
+	 * @param Volume Value between 0.0 and 1.0
+	 */
+	virtual void SetPlayerVolume(const FString& PlayerName, float Volume) = 0;
+
+	/**
+	 * @param PlayerName Player to get the volume of
+	 * @return volume for player
+	 */
+	virtual float GetPlayerVolume(const FString& PlayerName) const = 0;
+
+	/**
+	 * Delegate triggered when a player's volume is updated, usually as a result of calling SetPlayerVolume
+	 */
+	virtual FOnVoiceChatPlayerVolumeUpdatedDelegate& OnVoiceChatPlayerVolumeUpdated() = 0;
+
+	/**
+	 * Transmit to all channels that you are currently in
+	 */
+	virtual void TransmitToAllChannels() = 0;
+
+	/**
+	 * Stop transmitting to all channels that you are currently in. This does not apply to channels you will join.
+	 */
+	virtual void TransmitToNoChannels() = 0;
+
+	/**
+	 * Transmit to a specific channel
+	 *
+	 * @param ChannelName Channel to transmit to
+	 */
+	virtual void TransmitToSpecificChannel(const FString& ChannelName) = 0;
+
+	/**
+	 * Get the current transmit mode
+	 *
+	 * @return Transmit mode. If it is EVoiceChatTransmitMode::Channel, the channel can be retrieved using GetTransmitChannel()
+	 */
+	virtual EVoiceChatTransmitMode GetTransmitMode() const = 0;
+
+	/**
+	 * Get the specific channel we are transmitting to
+	 *
+	 * @return Channel we are transmitting to. Will be empty if GetTransmissionMode() returned a value other than EVoiceChatTransmitMode::Channel
+	 */
+	virtual FString GetTransmitChannel() const = 0;
+
+	/*
+	 * Start recording
+	 *
+	 * @param Delegate This will be called from audio thread.
+	 *
+	 * @return Delegate handle to pass to stop recording
+	 */
+	virtual FDelegateHandle StartRecording(const FOnVoiceChatRecordSamplesAvailableDelegate::FDelegate& Delegate) = 0;
+
+	/*
+	 * Stop recording
+	 *
+	 * @param Handle Delegate handle returned from a StartRecording call
+	 */
+	virtual void StopRecording(FDelegateHandle Handle) = 0;
+
+	/**
+	 * Delegates called from audio thread.
+	 */
+	virtual FDelegateHandle RegisterOnVoiceChatAfterCaptureAudioReadDelegate(const FOnVoiceChatAfterCaptureAudioReadDelegate::FDelegate& Delegate) = 0;
+	virtual void UnregisterOnVoiceChatAfterCaptureAudioReadDelegate(FDelegateHandle Handle) = 0;
+
+	virtual FDelegateHandle RegisterOnVoiceChatBeforeCaptureAudioSentDelegate(const FOnVoiceChatBeforeCaptureAudioSentDelegate::FDelegate& Delegate) = 0;
+	virtual void UnregisterOnVoiceChatBeforeCaptureAudioSentDelegate(FDelegateHandle Handle) = 0;
+
+	virtual FDelegateHandle RegisterOnVoiceChatBeforeRecvAudioRenderedDelegate(const FOnVoiceChatBeforeRecvAudioRenderedDelegate::FDelegate& Delegate) = 0;
+	virtual void UnregisterOnVoiceChatBeforeRecvAudioRenderedDelegate(FDelegateHandle Handle) = 0;
+
+	/**
+	 * Generate Login/Join tokens. Should be used for development purposes only
+	 */
+	virtual FString InsecureGetLoginToken(const FString& PlayerName) = 0;
+	virtual FString InsecureGetJoinToken(const FString& ChannelName, EVoiceChatChannelType ChannelType, TOptional<FVoiceChatChannel3dProperties> Channel3dProperties = TOptional<FVoiceChatChannel3dProperties>()) = 0;
+
+private:
+	static FName GetModularFeatureName()
+	{
+		static FName FeatureName = FName(TEXT("VoiceChat"));
+		return FeatureName;
+	}
+};
\ No newline at end of file
diff --git a/Engine/Plugins/Online/VoiceChat/VoiceChat/Source/Public/VoiceChatErrors.h b/Engine/Plugins/Online/VoiceChat/VoiceChat/Source/Public/VoiceChatErrors.h
new file mode 100644
--- /dev/null
+++ b/Engine/Plugins/Online/VoiceChat/VoiceChat/Source/Public/VoiceChatErrors.h
@@ -0,0 +1,32 @@
+// Copyright 1998-2019 Epic Games, Inc. All Rights Reserved.
+#pragma once
+
+#define VOICECHAT_ERROR(...) FVoiceChatResult::CreateError(TEXT("errors.com.epicgames.voicechat"), __VA_ARGS__)
+
+namespace VoiceChat
+{
+	namespace Errors
+	{
+		inline FVoiceChatResult NotInitialized() { return VOICECHAT_ERROR(EVoiceChatResult::NotInitialized, TEXT("not_initialized")); }
+		inline FVoiceChatResult NotConnected() { return VOICECHAT_ERROR(EVoiceChatResult::NotConnected, TEXT("not_connected")); }
+		inline FVoiceChatResult NotLoggedIn() { return VOICECHAT_ERROR(EVoiceChatResult::NotLoggedIn, TEXT("not_logged_in")); }
+
+		inline FVoiceChatResult InvalidState() { return VOICECHAT_ERROR(EVoiceChatResult::InvalidState, TEXT("invalid_state")); }
+		inline FVoiceChatResult DisconnectInProgress() { return VOICECHAT_ERROR(EVoiceChatResult::InvalidState, TEXT("disconnect_in_progress")); }
+		inline FVoiceChatResult OtherUserLoggedIn() { return VOICECHAT_ERROR(EVoiceChatResult::InvalidState, TEXT("other_user_logged_in")); }
+		inline FVoiceChatResult NotInChannel() { return VOICECHAT_ERROR(EVoiceChatResult::InvalidState, TEXT("not_in_channel")); }
+		inline FVoiceChatResult ChannelJoinInProgress() { return VOICECHAT_ERROR(EVoiceChatResult::InvalidState, TEXT("channel_join_in_progress")); }
+		inline FVoiceChatResult ChannelLeaveInProgress() { return VOICECHAT_ERROR(EVoiceChatResult::InvalidState, TEXT("channel_leave_in_progress")); }
+
+		inline FVoiceChatResult InvalidArgument(const FString& ErrorDesc = FString()) { return VOICECHAT_ERROR(EVoiceChatResult::InvalidArgument, TEXT("invalid_argument"), ErrorDesc); }
+		inline FVoiceChatResult CredentialsInvalid() { return VOICECHAT_ERROR(EVoiceChatResult::CredentialsInvalid, TEXT("credentials_invalid")); }
+		inline FVoiceChatResult CredentialsExpired() { return VOICECHAT_ERROR(EVoiceChatResult::CredentialsExpired, TEXT("credentials_expired")); }
+
+		inline FVoiceChatResult ClientTimeout() { return VOICECHAT_ERROR(EVoiceChatResult::ClientTimeout, TEXT("client_timeout")); }
+		inline FVoiceChatResult ServerTimeout() { return VOICECHAT_ERROR(EVoiceChatResult::ServerTimeout, TEXT("server_timeout")); }
+		inline FVoiceChatResult DnsFailure() { return VOICECHAT_ERROR(EVoiceChatResult::DnsFailure, TEXT("dns_failure")); }
+		inline FVoiceChatResult ConnectionFailure() { return VOICECHAT_ERROR(EVoiceChatResult::ConnectionFailure, TEXT("connection_failure")); }
+	}
+}
+
+#undef VOICECHAT_ERROR
\ No newline at end of file
diff --git a/Engine/Plugins/Online/VoiceChat/VoiceChat/Source/Public/VoiceChatResult.h b/Engine/Plugins/Online/VoiceChat/VoiceChat/Source/Public/VoiceChatResult.h
new file mode 100644
--- /dev/null
+++ b/Engine/Plugins/Online/VoiceChat/VoiceChat/Source/Public/VoiceChatResult.h
@@ -0,0 +1,122 @@
+// Copyright 1998-2019 Epic Games, Inc. All Rights Reserved.
+#pragma once
+
+enum class EVoiceChatResult
+{
+	// The operation succeeded
+	Success = 0,
+
+	// Common state errors
+	InvalidState,
+	NotInitialized,
+	NotConnected,
+	NotLoggedIn,
+
+	// Common argument errors
+	InvalidArgument,
+	CredentialsInvalid,
+	CredentialsExpired,
+
+	// Common connection errors
+	ClientTimeout,
+	ServerTimeout,
+	DnsFailure,
+	ConnectionFailure,
+
+	// Error does not map to any common categories
+	ImplementationError
+};
+
+struct FVoiceChatResult
+{
+	/** Whether or not we succeeded */
+	EVoiceChatResult ResultCode;
+	/** If we failed, the code for the error */
+	FString ErrorCode;
+	/** If we failed, a numeric error from the implementation */
+	int ErrorNum = 0;
+	/** If we failed, more details about the error condition */
+	FString ErrorDesc;
+
+	bool IsSuccess() const { return ResultCode == EVoiceChatResult::Success; }
+	
+	bool operator==(const FVoiceChatResult& Other) const
+	{
+		return ResultCode == Other.ResultCode && ErrorCode == Other.ErrorCode;
+	}
+
+	bool operator!=(const FVoiceChatResult& Other) const
+	{
+		return !(FVoiceChatResult::operator==(Other));
+	}
+
+	FVoiceChatResult(EVoiceChatResult ResultCode) : ResultCode(ResultCode) {}
+	FVoiceChatResult(EVoiceChatResult ResultCode, const FString& ErrorCode) : ResultCode(ResultCode), ErrorCode(ErrorCode) {}
+	FVoiceChatResult(EVoiceChatResult ResultCode, const FString& ErrorCode, const FString& ErrorDesc) : ResultCode(ResultCode), ErrorCode(ErrorCode), ErrorDesc(ErrorDesc) {}
+	FVoiceChatResult(EVoiceChatResult ResultCode, const FString& ErrorCode, int ErrorNum) : ResultCode(ResultCode), ErrorCode(ErrorCode), ErrorNum(ErrorNum) {}
+	FVoiceChatResult(EVoiceChatResult ResultCode, const FString& ErrorCode, int ErrorNum, const FString& ErrorDesc) : ResultCode(ResultCode), ErrorCode(ErrorCode), ErrorDesc(ErrorDesc) {}
+
+	static FVoiceChatResult CreateSuccess() { return FVoiceChatResult(EVoiceChatResult::Success); }
+
+	/** Create factory for proper namespacing */
+	static FVoiceChatResult CreateError(const FString& ErrorNamespace, EVoiceChatResult ResultCode, const FString& ErrorCode, const FString& ErrorDesc = FString())
+	{
+		FVoiceChatResult Error(ResultCode, ErrorCode, ErrorDesc);
+		if (!Error.ErrorCode.IsEmpty() && !Error.ErrorCode.StartsWith(TEXT("errors.com.")))
+		{
+			FString NamespacedError = ErrorNamespace;
+			NamespacedError += TEXT(".");
+			NamespacedError += MoveTemp(Error.ErrorCode);
+			Error.ErrorCode = MoveTemp(NamespacedError);
+		}
+
+		return Error;
+	}
+
+	/** Create factory for proper namespacing */
+	static FVoiceChatResult CreateError(const FString& ErrorNamespace, EVoiceChatResult ResultCode, const FString& ErrorCode, int ErrorNum, const FString& ErrorDesc = FString())
+	{
+		FVoiceChatResult Error = CreateError(ErrorNamespace, ResultCode, ErrorCode, ErrorDesc);
+		Error.ErrorNum = ErrorNum;
+		return Error;
+	}
+};
+
+inline FString LexToString(const EVoiceChatResult Result)
+{
+	switch (Result)
+	{
+	case EVoiceChatResult::Success: return TEXT("Success");
+	case EVoiceChatResult::InvalidState: return TEXT("InvalidState");
+	case EVoiceChatResult::NotInitialized: return TEXT("NotInitialized");
+	case EVoiceChatResult::NotConnected: return TEXT("NotConnected");
+	case EVoiceChatResult::NotLoggedIn: return TEXT("NotLoggedIn");
+	case EVoiceChatResult::InvalidArgument: return TEXT("InvalidArgument");
+	case EVoiceChatResult::CredentialsInvalid: return TEXT("AccessTokenInvalid");
+	case EVoiceChatResult::CredentialsExpired: return TEXT("AccessTokenExpired");
+	case EVoiceChatResult::ClientTimeout: return TEXT("ClientTimeout");
+	case EVoiceChatResult::ServerTimeout: return TEXT("ServerTimeout");
+	case EVoiceChatResult::DnsFailure: return TEXT("DnsFailure");
+	case EVoiceChatResult::ConnectionFailure: return TEXT("ConnectionFailure");
+	case EVoiceChatResult::ImplementationError: return TEXT("ImplementationError");
+	default:
+		checkNoEntry();
+		return TEXT("UNKNOWN");
+	}
+}
+
+inline FString LexToString(const FVoiceChatResult& Result)
+{
+	if (Result.IsSuccess())
+	{
+		return TEXT("Success");
+	}
+	else if (Result.ErrorDesc.IsEmpty())
+	{
+		return FString::Printf(TEXT("Failed with ResultCode=[%s], ErrorCode=[%s]"), *LexToString(Result.ResultCode), *Result.ErrorCode);
+	}
+	else
+	{
+		return FString::Printf(TEXT("Failed with ResultCode=[%s], ErrorCode=[%s], ErrorDesc=[%s]"), *LexToString(Result.ResultCode), *Result.ErrorCode, *Result.ErrorDesc);
+	}
+}
\ No newline at end of file
diff --git a/Engine/Plugins/Online/VoiceChat/VoiceChat/Source/VoiceChat.Build.cs b/Engine/Plugins/Online/VoiceChat/VoiceChat/Source/VoiceChat.Build.cs
new file mode 100644
--- /dev/null
+++ b/Engine/Plugins/Online/VoiceChat/VoiceChat/Source/VoiceChat.Build.cs
@@ -0,0 +1,21 @@
+// Copyright 1998-2019 Epic Games, Inc. All Rights Reserved.
+
+using System.IO;
+
+namespace UnrealBuildTool.Rules
+{
+	public class VoiceChat : ModuleRules
+	{
+		public VoiceChat(ReadOnlyTargetRules Target) : base(Target)
+		{
+			PCHUsage = ModuleRules.PCHUsageMode.UseExplicitOrSharedPCHs;
+
+			PublicDependencyModuleNames.AddRange(
+				new string[]
+				{
+					"Core"
+				}
+			);
+		}
+	}
+}
diff --git a/Engine/Plugins/Online/VoiceChat/VoiceChat/VoiceChat.uplugin b/Engine/Plugins/Online/VoiceChat/VoiceChat/VoiceChat.uplugin
new file mode 100644
--- /dev/null
+++ b/Engine/Plugins/Online/VoiceChat/VoiceChat/VoiceChat.uplugin
@@ -0,0 +1,20 @@
+{
+	"FileVersion" : 3,
+	
+	"FriendlyName" : "Voice Chat Interface",
+	"Version" : 1,
+	"VersionName" : "1.0",
+	"Description" : "Voice Chat Interface",
+	"Category" : "Online Platform",
+	"CreatedBy" : "Epic Games, Inc.",
+	"CreatedByURL" : "http://epicgames.com",
+	"EnabledByDefault" : false,
+	
+	"Modules" :
+	[
+		{
+			"Name" : "VoiceChat",
+			"Type" : "ClientOnly"
+		}
+	]
+}
diff --git a/Engine/Source/Runtime/Online/Voice/Public/Interfaces/VoiceChat.h b/Engine/Source/Runtime/Online/Voice/Public/Interfaces/VoiceChat.h
deleted file mode 100644
