Fix issue with setting volume to 0 not necessarily being muted

diff --git a/Engine/Plugins/Online/VoiceChat/VivoxVoiceChat/Source/Private/VivoxVoiceChat.cpp b/Engine/Plugins/Online/VoiceChat/VivoxVoiceChat/Source/Private/VivoxVoiceChat.cpp
--- a/Engine/Plugins/Online/VoiceChat/VivoxVoiceChat/Source/Private/VivoxVoiceChat.cpp
+++ b/Engine/Plugins/Online/VoiceChat/VivoxVoiceChat/Source/Private/VivoxVoiceChat.cpp
@@ -348,16 +348,16 @@ void FVivoxVoiceChat::SetAudioInputVolume(float InVolume)
 {
 	UE_LOG(LogVivoxVoiceChat, Verbose, TEXT("SetAudioInputVolume %f"), InVolume);
 
-	VivoxClientConnection.SetAudioInputDeviceMuted(InVolume <= SMALL_NUMBER);
-	VivoxClientConnection.SetMasterAudioInputDeviceVolume(FMath::Lerp(VIVOX_MIN_VOL, VIVOX_MAX_VOL, InVolume));
+	AudioInputOptions.Volume = InVolume;
+	ApplyAudioInputOptions();
 }
 
 void FVivoxVoiceChat::SetAudioOutputVolume(float InVolume)
 {
 	UE_LOG(LogVivoxVoiceChat, Verbose, TEXT("SetAudioOutputVolume %f"), InVolume);
 
-	VivoxClientConnection.SetAudioOutputDeviceMuted(InVolume <= SMALL_NUMBER);
-	VivoxClientConnection.SetMasterAudioOutputDeviceVolume(FMath::Lerp(VIVOX_MIN_VOL, VIVOX_MAX_VOL, InVolume));
+	AudioOutputOptions.Volume = InVolume;
+	ApplyAudioOutputOptions();
 }
 
 float FVivoxVoiceChat::GetAudioInputVolume() const
@@ -374,14 +374,16 @@ void FVivoxVoiceChat::SetAudioInputDeviceMuted(bool bIsMuted)
 {
 	UE_LOG(LogVivoxVoiceChat, Verbose, TEXT("SetAudioInputDeviceMuted %s"), *LexToString(bIsMuted));
 
-	VivoxClientConnection.SetAudioInputDeviceMuted(bIsMuted);
+	AudioInputOptions.bMuted = bIsMuted;
+	ApplyAudioInputOptions();
 }
 
 void FVivoxVoiceChat::SetAudioOutputDeviceMuted(bool bIsMuted)
 {
 	UE_LOG(LogVivoxVoiceChat, Verbose, TEXT("SetAudioOutputDeviceMuted %s"), *LexToString(bIsMuted));
 
-	VivoxClientConnection.SetAudioOutputDeviceMuted(bIsMuted);
+	AudioOutputOptions.bMuted = bIsMuted;
+	ApplyAudioOutputOptions();
 }
 
 bool FVivoxVoiceChat::GetAudioInputDeviceMuted() const
@@ -1935,6 +1937,22 @@ void FVivoxVoiceChat::ClearLoginSession()
 	LoginSession.State = FLoginSession::EState::LoggedOut;
 }
 
+void FVivoxVoiceChat::ApplyAudioInputOptions()
+{
+	VivoxClientConnection.SetMasterAudioInputDeviceVolume(FMath::Lerp(VIVOX_MIN_VOL, VIVOX_MAX_VOL, AudioInputOptions.Volume));
+
+	const bool bVolumeMuted = AudioInputOptions.Volume < SMALL_NUMBER;
+	VivoxClientConnection.SetAudioInputDeviceMuted(AudioInputOptions.bMuted || bVolumeMuted);
+}
+
+void FVivoxVoiceChat::ApplyAudioOutputOptions()
+{
+	VivoxClientConnection.SetMasterAudioOutputDeviceVolume(FMath::Lerp(VIVOX_MIN_VOL, VIVOX_MAX_VOL, AudioOutputOptions.Volume));
+
+	const bool bVolumeMuted = AudioOutputOptions.Volume < SMALL_NUMBER;
+	VivoxClientConnection.SetAudioOutputDeviceMuted(AudioOutputOptions.bMuted || bVolumeMuted);
+}
+
 bool FVivoxVoiceChat::Exec(UWorld* InWorld, const TCHAR* Cmd, FOutputDevice& Ar)
 {
 	if (FParse::Command(&Cmd, TEXT("VIVOX")))
diff --git a/Engine/Plugins/Online/VoiceChat/VivoxVoiceChat/Source/Public/VivoxVoiceChat.h b/Engine/Plugins/Online/VoiceChat/VivoxVoiceChat/Source/Public/VivoxVoiceChat.h
--- a/Engine/Plugins/Online/VoiceChat/VivoxVoiceChat/Source/Public/VivoxVoiceChat.h
+++ b/Engine/Plugins/Online/VoiceChat/VivoxVoiceChat/Source/Public/VivoxVoiceChat.h
@@ -205,6 +205,14 @@ protected:
 		TMap<FString, FParticipant> Participants; // Contains participants from all channels and the desired muted/volume state
 	};
 
+	struct AudioOptions
+	{
+		bool bMuted = false;
+		float Volume = 1.0f;
+	};
+	AudioOptions AudioInputOptions;
+	AudioOptions AudioOutputOptions;
+
 	enum class EConnectionState
 	{
 		Disconnected,
@@ -280,8 +288,9 @@ protected:
 	FChannelSession& GetChannelSession(const VivoxClientApi::Uri& ChannelUri);
 	void RemoveChannelSession(const FString& ChannelName);
 	void ClearChannelSessions();
-
 	void ClearLoginSession();
+	void ApplyAudioInputOptions();
+	void ApplyAudioOutputOptions();
 
 	// Log spam avoidance
 	FString LastLogMessage;
