VivoxVoiceChat Plugin - Add support for multiple local users

diff --git a/Engine/Plugins/Online/VoiceChat/VivoxVoiceChat/Source/Private/VivoxVoiceChat.cpp b/Engine/Plugins/Online/VoiceChat/VivoxVoiceChat/Source/Private/VivoxVoiceChat.cpp
--- a/Engine/Plugins/Online/VoiceChat/VivoxVoiceChat/Source/Private/VivoxVoiceChat.cpp
+++ b/Engine/Plugins/Online/VoiceChat/VivoxVoiceChat/Source/Private/VivoxVoiceChat.cpp
@@ -150,201 +150,28 @@
 
 }
 
-FVivoxVoiceChat::FVivoxVoiceChat()
-	: bInitialized(false)
-	, ConnectionState(EConnectionState::Disconnected)
+FVivoxVoiceChatUser::FVivoxVoiceChatUser(FVivoxVoiceChat& InVivoxVoiceChat)
+	: VivoxClientConnection(InVivoxVoiceChat.VivoxClientConnection) 
+	, VivoxVoiceChat(InVivoxVoiceChat)
 {
-	SetClientConnection(&VivoxClientConnection);
-	SetAbortEnabled(false);
+	VivoxVoiceChat.RegisterVoiceChatUser(this);
 }
 
-FVivoxVoiceChat::~FVivoxVoiceChat()
+FVivoxVoiceChatUser::~FVivoxVoiceChatUser()
+{
+	VivoxVoiceChat.UnregisterVoiceChatUser(this);
+}
+
+void FVivoxVoiceChatUser::SetSetting(const FString& Name, const FString& Value)
 {
 }
 
-static void* VivoxMalloc(size_t bytes)
+FString FVivoxVoiceChatUser::GetSetting(const FString& Name)
 {
-	LLM_SCOPE( ELLMTag::AudioVoiceChat );
-	return FMemory::Malloc(bytes);
-}
-
-static void VivoxFree(void* ptr)
-{
-	LLM_SCOPE( ELLMTag::AudioVoiceChat );
-	FMemory::Free(ptr);
-}
-
-static void* VivoxRealloc(void* ptr, size_t bytes)
-{
-	LLM_SCOPE( ELLMTag::AudioVoiceChat );
-	return FMemory::Realloc(ptr, bytes);
-}
-
-static void* VivoxCalloc(size_t num, size_t bytes)
-{
-	LLM_SCOPE( ELLMTag::AudioVoiceChat );
-	const size_t Size = bytes * num;
-	void* Ret = FMemory::Malloc(Size);
-	FMemory::Memzero(Ret, Size);
-	return Ret;
-}
-
-static void* VivoxMallocAligned(size_t alignment, size_t bytes)
-{
-	LLM_SCOPE( ELLMTag::AudioVoiceChat );
-	return FMemory::Malloc(bytes, alignment);
-}
-
-static void VivoxFreeAligned(void* ptr)
-{
-	LLM_SCOPE( ELLMTag::AudioVoiceChat );
-	FMemory::Free(ptr);
-}
-
-bool FVivoxVoiceChat::Initialize()
-{
-	if (!IsInitialized())
-	{
-		bool bEnabled = true;
-		GConfig->GetBool(TEXT("VoiceChat.Vivox"), TEXT("bEnabled"), bEnabled, GEngineIni);
-		if (bEnabled)
-		{
-			GConfig->GetString(TEXT("VoiceChat.Vivox"), TEXT("ServerUrl"), VivoxServerUrl, GEngineIni);
-			GConfig->GetString(TEXT("VoiceChat.Vivox"), TEXT("Domain"), VivoxDomain, GEngineIni);
-			GConfig->GetString(TEXT("VoiceChat.Vivox"), TEXT("Issuer"), VivoxIssuer, GEngineIni);
-			GConfig->GetString(TEXT("VoiceChat.Vivox"), TEXT("Namespace"), VivoxNamespace, GEngineIni);
-			GConfig->GetString(TEXT("VoiceChat.Vivox"), TEXT("InsecureSecret"), VivoxInsecureSecret, GEngineIni);
-
-			if (VivoxNamespace.IsEmpty())
-			{
-				VivoxNamespace = VivoxIssuer;
-			}
-
-			// positional audio settings
-			AttenuationModel = EVoiceChatAttenuationModel::InverseByDistance;
-			MinDistance = 100;
-			MaxDistance = 3000;
-			Rolloff = 1.0f;
-			FString AttenuationModelString;
-			if (GConfig->GetString(TEXT("VoiceChat.Vivox"), TEXT("AttenuationModel"), AttenuationModelString, GEngineIni) && !AttenuationModelString.IsEmpty())
-			{
-				if (AttenuationModelString == TEXT("None"))
-				{
-					AttenuationModel = EVoiceChatAttenuationModel::None;
-				}
-				else if (AttenuationModelString == TEXT("InverseByDistance"))
-				{
-					AttenuationModel = EVoiceChatAttenuationModel::InverseByDistance;
-				}
-				else if (AttenuationModelString == TEXT("LinearByDistance"))
-				{
-					AttenuationModel = EVoiceChatAttenuationModel::LinearByDistance;
-				}
-				else if (AttenuationModelString == TEXT("ExponentialByDistance"))
-				{
-					AttenuationModel = EVoiceChatAttenuationModel::ExponentialByDistance;
-				}
-				else
-				{
-					UE_LOG(LogVivoxVoiceChat, Warning, TEXT("Unknown AttenuationModel: %s"), *AttenuationModelString);
-				}
-			}
-			GConfig->GetInt(TEXT("VoiceChat.Vivox"), TEXT("MinDistance"), MinDistance, GEngineIni);
-			GConfig->GetInt(TEXT("VoiceChat.Vivox"), TEXT("MaxDistance"), MaxDistance, GEngineIni);
-			GConfig->GetFloat(TEXT("VoiceChat.Vivox"), TEXT("Rolloff"), Rolloff, GEngineIni);
-
-			const char* VivoxVersionInfo = vx_get_sdk_version_info();
-			UE_LOG(LogVivoxVoiceChat, Log, TEXT("Initializing Vivox %s"), ANSI_TO_TCHAR(VivoxVersionInfo));
-
-			vx_sdk_config_t ConfigHints;
-			int Result = vx_get_default_config3(&ConfigHints, sizeof(ConfigHints));
-			if (Result != 0)
-			{
-				UE_LOG(LogVivoxVoiceChat, Warning, TEXT("Failed to get default config: error:%s (%i)"), ANSI_TO_TCHAR(VivoxClientApi::GetErrorString(Result)), Result);
-			}
-			else
-			{
-				SetVivoxSdkConfigHints(ConfigHints);
-
-				LogLevel VivoxLogLevel = LogLevelWarning;
-				FString LogLevelString;
-				if (GConfig->GetString(TEXT("VoiceChat.Vivox"), TEXT("LogLevel"), LogLevelString, GEngineIni))
-				{
-					if (LogLevelString == TEXT("None"))
-					{
-						VivoxLogLevel = LogLevelNone;
-					}
-					else if (LogLevelString == TEXT("Error"))
-					{
-						VivoxLogLevel = LogLevelError;
-					}
-					else if (LogLevelString == TEXT("Warning"))
-					{
-						VivoxLogLevel = LogLevelWarning;
-					}
-					else if (LogLevelString == TEXT("Info"))
-					{
-						VivoxLogLevel = LogLevelInfo;
-					}
-					else if (LogLevelString == TEXT("Debug"))
-					{
-						VivoxLogLevel = LogLevelDebug;
-					}
-					else if (LogLevelString == TEXT("Trace"))
-					{
-						VivoxLogLevel = LogLevelTrace;
-					}
-				}
-
-				VivoxClientApi::VCSStatus Status = VivoxClientConnection.Initialize(this, VivoxLogLevel, true, false, &ConfigHints, sizeof(ConfigHints));
-				if (Status.IsError())
-				{
-					const FVoiceChatResult VoiceChatResult = ResultFromVivoxStatus(Status);
-					UE_LOG(LogVivoxVoiceChat, Warning, TEXT("Initialize %s"), *LexToString(VoiceChatResult));
-				}
-				else
-				{
-					bInitialized = true;
-
-					bool bVADAutomaticParameterSelection = true;
-					GConfig->GetBool(TEXT("VoiceChat.Vivox"), TEXT("bVADAutomaticParameterSelection"), bVADAutomaticParameterSelection, GEngineIni);
-					VivoxClientConnection.SetVADAutomaticParameterSelection(bVADAutomaticParameterSelection);
-				}
-			}
-		}
-	}
-
-	return IsInitialized();
-}
-
-bool FVivoxVoiceChat::Uninitialize()
-{
-	if (IsInitialized())
-	{
-		VivoxClientConnection.Uninitialize();
-		ConnectionState = EConnectionState::Disconnected;
-		LoginSession = FLoginSession();
-		bInitialized = false;
-	}
-
-	return true;
-}
-
-bool FVivoxVoiceChat::IsInitialized() const
-{
-	return bInitialized;
-}
-
-void FVivoxVoiceChat::SetSetting(const FString& Name, const FString& Value)
-{
-}
-
-FString FVivoxVoiceChat::GetSetting(const FString& Name)
-{
 	return FString();
 }
 
-void FVivoxVoiceChat::SetAudioInputVolume(float InVolume)
+void FVivoxVoiceChatUser::SetAudioInputVolume(float InVolume)
 {
 	UE_LOG(LogVivoxVoiceChat, Verbose, TEXT("SetAudioInputVolume %f"), InVolume);
 
@@ -352,7 +179,7 @@
 	ApplyAudioInputOptions();
 }
 
-void FVivoxVoiceChat::SetAudioOutputVolume(float InVolume)
+void FVivoxVoiceChatUser::SetAudioOutputVolume(float InVolume)
 {
 	UE_LOG(LogVivoxVoiceChat, Verbose, TEXT("SetAudioOutputVolume %f"), InVolume);
 
@@ -360,17 +187,17 @@
 	ApplyAudioOutputOptions();
 }
 
-float FVivoxVoiceChat::GetAudioInputVolume() const
+float FVivoxVoiceChatUser::GetAudioInputVolume() const
 {
-	return static_cast<float>(VivoxClientConnection.GetMasterAudioInputDeviceVolume() - VIVOX_MIN_VOL) / static_cast<float>(VIVOX_MAX_VOL - VIVOX_MIN_VOL);
+	return static_cast<float>(VivoxClientConnection.GetMasterAudioInputDeviceVolume(LoginSession.AccountName) - VIVOX_MIN_VOL) / static_cast<float>(VIVOX_MAX_VOL - VIVOX_MIN_VOL);
 }
 
-float FVivoxVoiceChat::GetAudioOutputVolume() const
+float FVivoxVoiceChatUser::GetAudioOutputVolume() const
 {
-	return static_cast<float>(VivoxClientConnection.GetMasterAudioOutputDeviceVolume() - VIVOX_MIN_VOL) / static_cast<float>(VIVOX_MAX_VOL - VIVOX_MIN_VOL);
+	return static_cast<float>(VivoxClientConnection.GetMasterAudioOutputDeviceVolume(LoginSession.AccountName) - VIVOX_MIN_VOL) / static_cast<float>(VIVOX_MAX_VOL - VIVOX_MIN_VOL);
 }
 
-void FVivoxVoiceChat::SetAudioInputDeviceMuted(bool bIsMuted)
+void FVivoxVoiceChatUser::SetAudioInputDeviceMuted(bool bIsMuted)
 {
 	UE_LOG(LogVivoxVoiceChat, Verbose, TEXT("SetAudioInputDeviceMuted %s"), *LexToString(bIsMuted));
 
@@ -378,7 +205,7 @@
 	ApplyAudioInputOptions();
 }
 
-void FVivoxVoiceChat::SetAudioOutputDeviceMuted(bool bIsMuted)
+void FVivoxVoiceChatUser::SetAudioOutputDeviceMuted(bool bIsMuted)
 {
 	UE_LOG(LogVivoxVoiceChat, Verbose, TEXT("SetAudioOutputDeviceMuted %s"), *LexToString(bIsMuted));
 
@@ -386,22 +213,22 @@
 	ApplyAudioOutputOptions();
 }
 
-bool FVivoxVoiceChat::GetAudioInputDeviceMuted() const
+bool FVivoxVoiceChatUser::GetAudioInputDeviceMuted() const
 {
-	return VivoxClientConnection.GetAudioInputDeviceMuted();
+	return VivoxClientConnection.GetAudioInputDeviceMuted(LoginSession.AccountName);
 }
 
-bool FVivoxVoiceChat::GetAudioOutputDeviceMuted() const
+bool FVivoxVoiceChatUser::GetAudioOutputDeviceMuted() const
 {
-	return VivoxClientConnection.GetAudioOutputDeviceMuted();
+	return VivoxClientConnection.GetAudioOutputDeviceMuted(LoginSession.AccountName);
 }
 
-TArray<FString> FVivoxVoiceChat::GetAvailableInputDevices() const
+TArray<FString> FVivoxVoiceChatUser::GetAvailableInputDevices() const
 {
 	TArray<FString> InputDevices;
 	const VivoxClientApi::AudioDeviceId* AudioDevices = nullptr;
 	int NumAudioDevices = 0;
-	VivoxClientConnection.GetAvailableAudioInputDevices(AudioDevices, NumAudioDevices);
+	VivoxClientConnection.GetAvailableAudioInputDevices(LoginSession.AccountName, AudioDevices, NumAudioDevices);
 	InputDevices.Reserve(NumAudioDevices);
 	if (AudioDevices)
 	{
@@ -413,12 +240,12 @@
 	return InputDevices;
 }
 
-TArray<FString> FVivoxVoiceChat::GetAvailableOutputDevices() const
+TArray<FString> FVivoxVoiceChatUser::GetAvailableOutputDevices() const
 {
 	TArray<FString> OutputDevices;
 	const VivoxClientApi::AudioDeviceId* AudioDevices = nullptr;
 	int NumAudioDevices = 0;
-	VivoxClientConnection.GetAvailableAudioOutputDevices(AudioDevices, NumAudioDevices);
+	VivoxClientConnection.GetAvailableAudioOutputDevices(LoginSession.AccountName, AudioDevices, NumAudioDevices);
 	OutputDevices.Reserve(NumAudioDevices);
 	if (AudioDevices)
 	{
@@ -430,7 +257,7 @@
 	return OutputDevices;
 }
 
-void FVivoxVoiceChat::SetInputDevice(const FString& InputDevice)
+void FVivoxVoiceChatUser::SetInputDevice(const FString& InputDevice)
 {
 	UE_LOG(LogVivoxVoiceChat, Verbose, TEXT("SetInputDevice %s"), *InputDevice);
 
@@ -438,7 +265,7 @@
 	{
 		const VivoxClientApi::AudioDeviceId* AudioDevices = nullptr;
 		int NumAudioDevices = 0;
-		VivoxClientConnection.GetAvailableAudioInputDevices(AudioDevices, NumAudioDevices);
+		VivoxClientConnection.GetAvailableAudioInputDevices(LoginSession.AccountName, AudioDevices, NumAudioDevices);
 		if (AudioDevices)
 		{
 			for (int DeviceIndex = 0; DeviceIndex < NumAudioDevices; ++DeviceIndex)
@@ -446,17 +273,17 @@
 				const VivoxClientApi::AudioDeviceId& DeviceId = AudioDevices[DeviceIndex];
 				if (InputDevice == UTF8_TO_TCHAR(DeviceId.GetAudioDeviceDisplayName()))
 				{
-					VivoxClientConnection.SetApplicationChosenAudioInputDevice(DeviceId);
+					VivoxClientConnection.SetApplicationChosenAudioInputDevice(LoginSession.AccountName, DeviceId);
 					return;
 				}
 			}
 		}
 	}
 
-	VivoxClientConnection.UseOperatingSystemChosenAudioInputDevice();
+	VivoxClientConnection.UseOperatingSystemChosenAudioInputDevice(LoginSession.AccountName);
 }
 
-void FVivoxVoiceChat::SetOutputDevice(const FString& OutputDevice)
+void FVivoxVoiceChatUser::SetOutputDevice(const FString& OutputDevice)
 {
 	UE_LOG(LogVivoxVoiceChat, Verbose, TEXT("SetOutputDevice %s"), *OutputDevice);
 
@@ -464,7 +291,7 @@
 	{
 		const VivoxClientApi::AudioDeviceId* AudioDevices = nullptr;
 		int NumAudioDevices = 0;
-		VivoxClientConnection.GetAvailableAudioOutputDevices(AudioDevices, NumAudioDevices);
+		VivoxClientConnection.GetAvailableAudioOutputDevices(LoginSession.AccountName, AudioDevices, NumAudioDevices);
 		if (AudioDevices)
 		{
 			for (int DeviceIndex = 0; DeviceIndex < NumAudioDevices; ++DeviceIndex)
@@ -472,57 +299,92 @@
 				const VivoxClientApi::AudioDeviceId& DeviceId = AudioDevices[DeviceIndex];
 				if (OutputDevice == UTF8_TO_TCHAR(DeviceId.GetAudioDeviceDisplayName()))
 				{
-					VivoxClientConnection.SetApplicationChosenAudioOutputDevice(DeviceId);
+					VivoxClientConnection.SetApplicationChosenAudioOutputDevice(LoginSession.AccountName, DeviceId);
 					return;
 				}
 			}
 		}
 	}
 
-	VivoxClientConnection.UseOperatingSystemChosenAudioOutputDevice();
+	VivoxClientConnection.UseOperatingSystemChosenAudioOutputDevice(LoginSession.AccountName);
 }
 
-FString FVivoxVoiceChat::GetInputDevice() const
+FString FVivoxVoiceChatUser::GetInputDevice() const
 {
-	if (VivoxClientConnection.IsUsingOperatingSystemChosenAudioInputDevice())
+	if (VivoxClientConnection.IsUsingOperatingSystemChosenAudioInputDevice(LoginSession.AccountName))
 	{
-		VivoxClientApi::AudioDeviceId DeviceId = VivoxClientConnection.GetOperatingSystemChosenAudioInputDevice();
+		VivoxClientApi::AudioDeviceId DeviceId = VivoxClientConnection.GetOperatingSystemChosenAudioInputDevice(LoginSession.AccountName);
 		return UTF8_TO_TCHAR(DeviceId.GetAudioDeviceDisplayName());
 	}
 	else
 	{
-		VivoxClientApi::AudioDeviceId DeviceId = VivoxClientConnection.GetApplicationChosenAudioInputDevice();
+		VivoxClientApi::AudioDeviceId DeviceId = VivoxClientConnection.GetApplicationChosenAudioInputDevice(LoginSession.AccountName);
 		return UTF8_TO_TCHAR(DeviceId.GetAudioDeviceDisplayName());
 	}
 }
 
-FString FVivoxVoiceChat::GetOutputDevice() const
+FString FVivoxVoiceChatUser::GetOutputDevice() const
 {
-	if (VivoxClientConnection.IsUsingOperatingSystemChosenAudioOutputDevice())
+	if (VivoxClientConnection.IsUsingOperatingSystemChosenAudioOutputDevice(LoginSession.AccountName))
 	{
-		VivoxClientApi::AudioDeviceId DeviceId = VivoxClientConnection.GetOperatingSystemChosenAudioOutputDevice();
+		VivoxClientApi::AudioDeviceId DeviceId = VivoxClientConnection.GetOperatingSystemChosenAudioOutputDevice(LoginSession.AccountName);
 		return UTF8_TO_TCHAR(DeviceId.GetAudioDeviceDisplayName());
 	}
 	else
 	{
-		VivoxClientApi::AudioDeviceId DeviceId = VivoxClientConnection.GetApplicationChosenAudioOutputDevice();
+		VivoxClientApi::AudioDeviceId DeviceId = VivoxClientConnection.GetApplicationChosenAudioOutputDevice(LoginSession.AccountName);
 		return UTF8_TO_TCHAR(DeviceId.GetAudioDeviceDisplayName());
 	}
 }
 
-FString FVivoxVoiceChat::GetDefaultInputDevice() const
+FString FVivoxVoiceChatUser::GetDefaultInputDevice() const
 {
-	VivoxClientApi::AudioDeviceId DeviceId = VivoxClientConnection.GetOperatingSystemChosenAudioInputDevice();
+	VivoxClientApi::AudioDeviceId DeviceId = VivoxClientConnection.GetOperatingSystemChosenAudioInputDevice(LoginSession.AccountName);
 	return UTF8_TO_TCHAR(DeviceId.GetAudioDeviceDisplayName());
 }
 
-FString FVivoxVoiceChat::GetDefaultOutputDevice() const
+FString FVivoxVoiceChatUser::GetDefaultOutputDevice() const
 {
-	VivoxClientApi::AudioDeviceId DeviceId = VivoxClientConnection.GetOperatingSystemChosenAudioOutputDevice();
+	VivoxClientApi::AudioDeviceId DeviceId = VivoxClientConnection.GetOperatingSystemChosenAudioOutputDevice(LoginSession.AccountName);
 	return UTF8_TO_TCHAR(DeviceId.GetAudioDeviceDisplayName());
 }
 
-void FVivoxVoiceChat::Connect(const FOnVoiceChatConnectCompleteDelegate& Delegate)
+void FVivoxVoiceChatUser::Connect(const FOnVoiceChatConnectCompleteDelegate& Delegate)
 {
+	VivoxVoiceChat.Connect(Delegate);
+}
+
+void FVivoxVoiceChatUser::Disconnect(const FOnVoiceChatDisconnectCompleteDelegate& Delegate)
+{
+	VivoxVoiceChat.Disconnect(Delegate);
+}
+
+bool FVivoxVoiceChatUser::IsConnecting() const
+{
+	return VivoxVoiceChat.IsConnecting();
+}
+
+bool FVivoxVoiceChatUser::IsConnected() const
+{
+	return VivoxVoiceChat.IsConnected();
+}
+
+FOnVoiceChatConnectedDelegate& FVivoxVoiceChatUser::OnVoiceChatConnected()
+{
+	return VivoxVoiceChat.OnVoiceChatConnected();
+}
+
+FOnVoiceChatDisconnectedDelegate& FVivoxVoiceChatUser::OnVoiceChatDisconnected()
+{
+	return VivoxVoiceChat.OnVoiceChatDisconnected();
+}
+
+FOnVoiceChatReconnectedDelegate& FVivoxVoiceChatUser::OnVoiceChatReconnected()
+{
+	return VivoxVoiceChat.OnVoiceChatReconnected();
+}
+
+void FVivoxVoiceChatUser::Login(FPlatformUserId PlatformId, const FString& PlayerName, const FString& Credentials, const FOnVoiceChatLoginCompleteDelegate& Delegate)
+{
 	FVoiceChatResult Result = FVoiceChatResult::CreateSuccess();
 
@@ -530,93 +392,6 @@
 	{
 		Result = VoiceChat::Errors::NotInitialized();
 	}
-	else if (ConnectionState == EConnectionState::Disconnecting)
-	{
-		Result = VoiceChat::Errors::DisconnectInProgress();
-	}
-
-	if (!Result.IsSuccess())
-	{
-		UE_LOG(LogVivoxVoiceChat, Warning, TEXT("Connect %s"), *LexToString(Result));
-
-		Delegate.ExecuteIfBound(Result);
-	}
-	else if (IsConnected())
-	{
-		Delegate.ExecuteIfBound(FVoiceChatResult::CreateSuccess());
-	}
-	else
-	{
-		OnVoiceChatConnectCompleteDelegates.Add(Delegate);
-
-		if (!IsConnecting())
-		{
-			if (VivoxServerUrl.IsEmpty() || VivoxDomain.IsEmpty() || VivoxNamespace.IsEmpty())
-			{
-				UE_LOG(LogVivoxVoiceChat, Warning, TEXT("[VoiceChat.Vivox] ServerUrl, Domain, or Issuer is not set. Vivox voice chat will not work"));
-				Result = VoiceChat::Errors::MissingConfig();
-				Result.ErrorDesc = FString::Printf(TEXT("ServerUrl=[%s] Domain=[%s] Namespace=[%s]"), *VivoxServerUrl, *VivoxDomain, *VivoxNamespace);
-			}
-			else
-			{
-				VivoxClientApi::Uri BackendUri(TCHAR_TO_ANSI(*VivoxServerUrl));
-				VivoxClientApi::VCSStatus Status = VivoxClientConnection.Connect(BackendUri);
-
-				if (Status.IsError())
-				{
-					UE_LOG(LogVivoxVoiceChat, Warning, TEXT("Connect failed: server:%s error:%s (%i)"), ANSI_TO_TCHAR(BackendUri.ToString()), ANSI_TO_TCHAR(Status.ToString()), Status.GetStatusCode());
-					Result = ResultFromVivoxStatus(Status);
-				}
-				else
-				{
-					ConnectionState = EConnectionState::Connecting;
-				}
-			}
-
-			if (!Result.IsSuccess())
-			{
-				ConnectionState = EConnectionState::Disconnected;
-
-				TriggerCompletionDelegates(OnVoiceChatConnectCompleteDelegates, Result);
-			}
-		}
-	}
-}
-
-void FVivoxVoiceChat::Disconnect(const FOnVoiceChatDisconnectCompleteDelegate& Delegate)
-{
-	if (IsConnected())
-	{
-		VivoxClientApi::Uri BackendUri(TCHAR_TO_ANSI(*VivoxServerUrl));
-		VivoxClientConnection.Disconnect(BackendUri);
-
-		ConnectionState = EConnectionState::Disconnecting;
-		OnVoiceChatDisconnectCompleteDelegates.Add(Delegate);
-	}
-	else
-	{
-		Delegate.ExecuteIfBound(FVoiceChatResult::CreateSuccess());
-	}
-}
-
-bool FVivoxVoiceChat::IsConnecting() const
-{
-	return ConnectionState == EConnectionState::Connecting;
-}
-
-bool FVivoxVoiceChat::IsConnected() const
-{
-	return ConnectionState == EConnectionState::Connected;
-}
-
-void FVivoxVoiceChat::Login(FPlatformUserId PlatformId, const FString& PlayerName, const FString& Credentials, const FOnVoiceChatLoginCompleteDelegate& Delegate)
-{
-	FVoiceChatResult Result = FVoiceChatResult::CreateSuccess();
-
-	if (!IsInitialized())
-	{
-		Result = VoiceChat::Errors::NotInitialized();
-	}
 	else if (!IsConnected())
 	{
 		Result = VoiceChat::Errors::NotConnected();
@@ -641,7 +416,7 @@
 	{
 		Result = VoiceChat::Errors::InvalidArgument(TEXT("PlayerName invalid characters"));
 	}
-	else if(PlayerName.Len() > 60 - VivoxNamespace.Len())
+	else if(PlayerName.Len() > 60 - VivoxVoiceChat.VivoxNamespace.Len())
 	{
 		// Name must be between 3-63 characters long and start and end with a '.'. It also must contain the issuer and another '.' separating issuer and name
 		Result = VoiceChat::Errors::InvalidArgument(TEXT("PlayerName too long"));
@@ -654,8 +429,8 @@
 		return;
 	}
 
-	const VivoxClientApi::AccountName AccountName = CreateAccountName(PlayerName);
-	const VivoxClientApi::Uri UserUri = CreateUserUri(PlayerName);
+	const VivoxClientApi::AccountName AccountName = VivoxVoiceChat.CreateAccountName(PlayerName);
+	const VivoxClientApi::Uri UserUri = VivoxVoiceChat.CreateUserUri(PlayerName);
 
 	if (!AccountName.IsValid() || !UserUri.IsValid())
 	{
@@ -685,7 +460,7 @@
 	OnVoiceChatLoginCompleteDelegate = Delegate;
 }
 
-void FVivoxVoiceChat::Logout(const FOnVoiceChatLogoutCompleteDelegate& Delegate)
+void FVivoxVoiceChatUser::Logout(const FOnVoiceChatLogoutCompleteDelegate& Delegate)
 {
 	FVoiceChatResult Result = FVoiceChatResult::CreateSuccess();
 
@@ -710,7 +485,7 @@
 		return;
 	}
 
-	VivoxClientApi::AccountName AccountName = CreateAccountName(LoginSession.PlayerName);
+	VivoxClientApi::AccountName AccountName = VivoxVoiceChat.CreateAccountName(LoginSession.PlayerName);
 	VivoxClientApi::VCSStatus Status = VivoxClientConnection.Logout(AccountName);
 	if (Status.IsError())
 	{
@@ -724,27 +499,27 @@
 	LoginSession.State = FLoginSession::EState::LoggingOut;
 }
 
-bool FVivoxVoiceChat::IsLoggingIn() const
+bool FVivoxVoiceChatUser::IsLoggingIn() const
 {
 	return LoginSession.State == FLoginSession::EState::LoggingIn;
 }
 
-bool FVivoxVoiceChat::IsLoggedIn() const
+bool FVivoxVoiceChatUser::IsLoggedIn() const
 {
 	return LoginSession.State == FLoginSession::EState::LoggedIn;
 }
 
-FString FVivoxVoiceChat::GetLoggedInPlayerName() const
+FString FVivoxVoiceChatUser::GetLoggedInPlayerName() const
 {
 	return IsLoggedIn() ? LoginSession.PlayerName : FString();
 }
 
-void FVivoxVoiceChat::BlockPlayers(const TArray<FString>& PlayerNames)
+void FVivoxVoiceChatUser::BlockPlayers(const TArray<FString>& PlayerNames)
 {
 	TArray<VivoxClientApi::Uri> UsersToBlock;
 	for (const FString& PlayerName : PlayerNames)
 	{
-		UsersToBlock.Add(CreateUserUri(PlayerName));
+		UsersToBlock.Add(VivoxVoiceChat.CreateUserUri(PlayerName));
 	}
 
 	VivoxClientApi::VCSStatus Status = VivoxClientConnection.BlockUsers(LoginSession.AccountName, UsersToBlock.GetData(), UsersToBlock.Num());
@@ -761,12 +536,12 @@
 	}
 }
 
-void FVivoxVoiceChat::UnblockPlayers(const TArray<FString>& PlayerNames)
+void FVivoxVoiceChatUser::UnblockPlayers(const TArray<FString>& PlayerNames)
 {
 	TArray<VivoxClientApi::Uri> UsersToUnblock;
 	for (const FString& PlayerName : PlayerNames)
 	{
-		UsersToUnblock.Add(CreateUserUri(PlayerName));
+		UsersToUnblock.Add(VivoxVoiceChat.CreateUserUri(PlayerName));
 	}
 
 	VivoxClientApi::VCSStatus Status = VivoxClientConnection.UnblockUsers(LoginSession.AccountName, UsersToUnblock.GetData(), UsersToUnblock.Num());
@@ -783,7 +558,7 @@
 	}
 }
 
-void FVivoxVoiceChat::JoinChannel(const FString& ChannelName, const FString& ChannelCredentials, EVoiceChatChannelType ChannelType, const FOnVoiceChatChannelJoinCompleteDelegate& Delegate, TOptional<FVoiceChatChannel3dProperties> Channel3dProperties)
+void FVivoxVoiceChatUser::JoinChannel(const FString& ChannelName, const FString& ChannelCredentials, EVoiceChatChannelType ChannelType, const FOnVoiceChatChannelJoinCompleteDelegate& Delegate, TOptional<FVoiceChatChannel3dProperties> Channel3dProperties)
 {
 	FVoiceChatResult Result = FVoiceChatResult::CreateSuccess();
 
@@ -807,7 +582,7 @@
 	{
 		Result = VoiceChat::Errors::InvalidArgument(TEXT("ChannelName invalid characters"));
 	}
-	else if (ChannelName.Len() > 189 - VivoxNamespace.Len())
+	else if (ChannelName.Len() > 189 - VivoxVoiceChat.VivoxNamespace.Len())
 	{
 		// channel name length must not exceed 200 characters, including the confctl-?- prefix and the issuer and separator
 		Result = VoiceChat::Errors::InvalidArgument(TEXT("ChannelName too long"));
@@ -839,7 +614,7 @@
 
 	ChannelSession.ChannelName = ChannelName;
 	ChannelSession.ChannelType = ChannelType;
-	ChannelSession.ChannelUri = CreateChannelUri(ChannelName, ChannelType, Channel3dProperties);
+	ChannelSession.ChannelUri = VivoxVoiceChat.CreateChannelUri(ChannelName, ChannelType, Channel3dProperties);
 
 	VivoxClientApi::VCSStatus Status = VivoxClientConnection.JoinChannel(LoginSession.AccountName, ChannelSession.ChannelUri, TCHAR_TO_ANSI(*ChannelCredentials));
 	if (Status.IsError())
@@ -854,7 +629,7 @@
 	ChannelSession.JoinDelegate = Delegate;
 }
 
-void FVivoxVoiceChat::LeaveChannel(const FString& ChannelName, const FOnVoiceChatChannelLeaveCompleteDelegate& Delegate)
+void FVivoxVoiceChatUser::LeaveChannel(const FString& ChannelName, const FOnVoiceChatChannelLeaveCompleteDelegate& Delegate)
 {
 	FVoiceChatResult Result = FVoiceChatResult::CreateSuccess();
 
@@ -901,7 +676,7 @@
 	ChannelSession.LeaveDelegate = Delegate;
 }
 
-void FVivoxVoiceChat::Set3DPosition(const FString& ChannelName, const FVector& SpeakerPosition, const FVector& ListenerPosition, const FVector& ListenerForwardDirection, const FVector& ListenerUpDirection)
+void FVivoxVoiceChatUser::Set3DPosition(const FString& ChannelName, const FVector& SpeakerPosition, const FVector& ListenerPosition, const FVector& ListenerForwardDirection, const FVector& ListenerUpDirection)
 {
 	FChannelSession& ChannelSession = GetChannelSession(ChannelName);
 
@@ -917,7 +692,7 @@
 	}
 }
 
-TArray<FString> FVivoxVoiceChat::GetChannels() const
+TArray<FString> FVivoxVoiceChatUser::GetChannels() const
 {
 	TArray<FString> ChannelNames;
 	for (const TPair<FString, FChannelSession>& ChannelSessionPair : LoginSession.ChannelSessions)
@@ -932,24 +707,24 @@
 	return ChannelNames;
 }
 
-TArray<FString> FVivoxVoiceChat::GetPlayersInChannel(const FString& ChannelName) const
+TArray<FString> FVivoxVoiceChatUser::GetPlayersInChannel(const FString& ChannelName) const
 {
 	TArray<FString> PlayerNames;
 	GetChannelSession(ChannelName).Participants.GenerateKeyArray(PlayerNames);
 	return PlayerNames;
 }
 
-EVoiceChatChannelType FVivoxVoiceChat::GetChannelType(const FString& ChannelName) const
+EVoiceChatChannelType FVivoxVoiceChatUser::GetChannelType(const FString& ChannelName) const
 {
 	return GetChannelSession(ChannelName).ChannelType;
 }
 
-bool FVivoxVoiceChat::IsPlayerTalking(const FString& PlayerName) const
+bool FVivoxVoiceChatUser::IsPlayerTalking(const FString& PlayerName) const
 {
 	return GetParticipant(PlayerName).bTalking;
 }
 
-void FVivoxVoiceChat::SetPlayerMuted(const FString& PlayerName, bool bMuted)
+void FVivoxVoiceChatUser::SetPlayerMuted(const FString& PlayerName, bool bMuted)
 {
 	FParticipant& Participant = GetParticipant(PlayerName);
 	Participant.bMuted = bMuted;
@@ -973,12 +748,12 @@
 	}
 }
 
-bool FVivoxVoiceChat::IsPlayerMuted(const FString& PlayerName) const
+bool FVivoxVoiceChatUser::IsPlayerMuted(const FString& PlayerName) const
 {
 	return GetParticipant(PlayerName).bMuted;
 }
 
-void FVivoxVoiceChat::SetPlayerVolume(const FString& PlayerName, float Volume)
+void FVivoxVoiceChatUser::SetPlayerVolume(const FString& PlayerName, float Volume)
 {
 	FParticipant& Participant = GetParticipant(PlayerName);
 	Participant.Volume = FMath::Clamp(Volume, 0.0f, 1.0f);
@@ -1010,12 +785,12 @@
 	}
 }
 
-float FVivoxVoiceChat::GetPlayerVolume(const FString& PlayerName) const
+float FVivoxVoiceChatUser::GetPlayerVolume(const FString& PlayerName) const
 {
 	return GetParticipant(PlayerName).Volume;
 }
 
-void FVivoxVoiceChat::TransmitToNoChannels()
+void FVivoxVoiceChatUser::TransmitToNoChannels()
 {
 	UE_LOG(LogVivoxVoiceChat, Log, TEXT("TransmitToNoChannels"));
 
@@ -1026,7 +801,7 @@
 	}
 }
 
-void FVivoxVoiceChat::TransmitToAllChannels()
+void FVivoxVoiceChatUser::TransmitToAllChannels()
 {
 	UE_LOG(LogVivoxVoiceChat, Log, TEXT("TransmitToAllChannels"));
 
@@ -1037,7 +812,7 @@
 	}
 }
 
-void FVivoxVoiceChat::TransmitToSpecificChannel(const FString& Channel)
+void FVivoxVoiceChatUser::TransmitToSpecificChannel(const FString& Channel)
 {
 	UE_LOG(LogVivoxVoiceChat, Log, TEXT("TransmitToSpecificChannel %s"), *Channel);
 
@@ -1052,7 +827,7 @@
 	}
 }
 
-EVoiceChatTransmitMode FVivoxVoiceChat::GetTransmitMode() const
+EVoiceChatTransmitMode FVivoxVoiceChatUser::GetTransmitMode() const
 {
 	VivoxClientApi::ChannelTransmissionPolicy TransmissionPolicy = VivoxClientConnection.GetChannelTransmissionPolicy(LoginSession.AccountName);
 	switch (TransmissionPolicy.GetChannelTransmissionPolicy())
@@ -1067,7 +842,7 @@
 	}
 }
 
-FString FVivoxVoiceChat::GetTransmitChannel() const
+FString FVivoxVoiceChatUser::GetTransmitChannel() const
 {
 	VivoxClientApi::ChannelTransmissionPolicy TransmissionPolicy = VivoxClientConnection.GetChannelTransmissionPolicy(LoginSession.AccountName);
 	if (TransmissionPolicy.GetChannelTransmissionPolicy() == VivoxClientApi::ChannelTransmissionPolicy::vx_channel_transmission_policy_specific_channel)
@@ -1080,7 +855,7 @@
 	}
 }
 
-FDelegateHandle FVivoxVoiceChat::StartRecording(const FOnVoiceChatRecordSamplesAvailableDelegate::FDelegate& Delegate)
+FDelegateHandle FVivoxVoiceChatUser::StartRecording(const FOnVoiceChatRecordSamplesAvailableDelegate::FDelegate& Delegate)
 {
 	if (!VivoxClientConnection.AudioInputDeviceTestIsRecording())
 	{
@@ -1097,7 +872,7 @@
 	return OnVoiceChatRecordSamplesAvailableDelegate.Add(Delegate);
 }
 
-void FVivoxVoiceChat::StopRecording(FDelegateHandle Handle)
+void FVivoxVoiceChatUser::StopRecording(FDelegateHandle Handle)
 {
 	FScopeLock Lock(&AudioRecordLock);
 
@@ -1109,56 +884,56 @@
 	}
 }
 
-FDelegateHandle FVivoxVoiceChat::RegisterOnVoiceChatAfterCaptureAudioReadDelegate(const FOnVoiceChatAfterCaptureAudioReadDelegate::FDelegate& Delegate)
+FDelegateHandle FVivoxVoiceChatUser::RegisterOnVoiceChatAfterCaptureAudioReadDelegate(const FOnVoiceChatAfterCaptureAudioReadDelegate::FDelegate& Delegate)
 {
 	FScopeLock Lock(&AfterCaptureAudioReadLock);
 
 	return OnVoiceChatAfterCaptureAudioReadDelegate.Add(Delegate);
 }
 
-void FVivoxVoiceChat::UnregisterOnVoiceChatAfterCaptureAudioReadDelegate(FDelegateHandle Handle)
+void FVivoxVoiceChatUser::UnregisterOnVoiceChatAfterCaptureAudioReadDelegate(FDelegateHandle Handle)
 {
 	FScopeLock Lock(&AfterCaptureAudioReadLock);
 
 	OnVoiceChatAfterCaptureAudioReadDelegate.Remove(Handle);
 }
 
-FDelegateHandle FVivoxVoiceChat::RegisterOnVoiceChatBeforeCaptureAudioSentDelegate(const FOnVoiceChatBeforeCaptureAudioSentDelegate::FDelegate& Delegate)
+FDelegateHandle FVivoxVoiceChatUser::RegisterOnVoiceChatBeforeCaptureAudioSentDelegate(const FOnVoiceChatBeforeCaptureAudioSentDelegate::FDelegate& Delegate)
 {
 	FScopeLock Lock(&BeforeCaptureAudioSentLock);
 
 	return OnVoiceChatBeforeCaptureAudioSentDelegate.Add(Delegate);
 }
 
-void FVivoxVoiceChat::UnregisterOnVoiceChatBeforeCaptureAudioSentDelegate(FDelegateHandle Handle)
+void FVivoxVoiceChatUser::UnregisterOnVoiceChatBeforeCaptureAudioSentDelegate(FDelegateHandle Handle)
 {
 	FScopeLock Lock(&BeforeCaptureAudioSentLock);
 
 	OnVoiceChatBeforeCaptureAudioSentDelegate.Remove(Handle);
 }
 
-FDelegateHandle FVivoxVoiceChat::RegisterOnVoiceChatBeforeRecvAudioRenderedDelegate(const FOnVoiceChatBeforeRecvAudioRenderedDelegate::FDelegate& Delegate)
+FDelegateHandle FVivoxVoiceChatUser::RegisterOnVoiceChatBeforeRecvAudioRenderedDelegate(const FOnVoiceChatBeforeRecvAudioRenderedDelegate::FDelegate& Delegate)
 {
 	FScopeLock Lock(&BeforeRecvAudioRenderedLock);
 
 	return OnVoiceChatBeforeRecvAudioRenderedDelegate.Add(Delegate);
 }
 
-void FVivoxVoiceChat::UnregisterOnVoiceChatBeforeRecvAudioRenderedDelegate(FDelegateHandle Handle)
+void FVivoxVoiceChatUser::UnregisterOnVoiceChatBeforeRecvAudioRenderedDelegate(FDelegateHandle Handle)
 {
 	FScopeLock Lock(&BeforeRecvAudioRenderedLock);
 
 	OnVoiceChatBeforeRecvAudioRenderedDelegate.Remove(Handle);
 }
 
-FString FVivoxVoiceChat::InsecureGetLoginToken(const FString& PlayerName)
+FString FVivoxVoiceChatUser::InsecureGetLoginToken(const FString& PlayerName)
 {
 	FString Token;
 
 	if (IsInitialized())
 	{
-		VivoxClientApi::Uri UserUri = CreateUserUri(PlayerName);
-		if (char* ANSIToken = vx_debug_generate_token(TCHAR_TO_ANSI(*VivoxIssuer), FDateTime::UtcNow().ToUnixTimestamp() + 90, "login", FMath::Rand(), nullptr, UserUri.ToString(), nullptr, (const unsigned char*)TCHAR_TO_ANSI(*VivoxInsecureSecret), VivoxInsecureSecret.Len()))
+		VivoxClientApi::Uri UserUri = VivoxVoiceChat.CreateUserUri(PlayerName);
+		if (char* ANSIToken = vx_debug_generate_token(TCHAR_TO_ANSI(*VivoxVoiceChat.VivoxIssuer), FDateTime::UtcNow().ToUnixTimestamp() + 90, "login", FMath::Rand(), nullptr, UserUri.ToString(), nullptr, (const unsigned char*)TCHAR_TO_ANSI(*VivoxVoiceChat.VivoxInsecureSecret), VivoxVoiceChat.VivoxInsecureSecret.Len()))
 		{
 			Token = ANSI_TO_TCHAR(ANSIToken);
 			vx_free(ANSIToken);
@@ -1168,14 +943,14 @@
 	return Token;
 }
 
-FString FVivoxVoiceChat::InsecureGetJoinToken(const FString& ChannelName, EVoiceChatChannelType ChannelType, TOptional<FVoiceChatChannel3dProperties> Channel3dProperties)
+FString FVivoxVoiceChatUser::InsecureGetJoinToken(const FString& ChannelName, EVoiceChatChannelType ChannelType, TOptional<FVoiceChatChannel3dProperties> Channel3dProperties)
 {
 	FString Token;
 
 	if (IsInitialized() && IsLoggedIn())
 	{
-		VivoxClientApi::Uri ChannelUri = CreateChannelUri(ChannelName, ChannelType, Channel3dProperties);
-		if (char* ANSIToken = vx_debug_generate_token(TCHAR_TO_ANSI(*VivoxIssuer), FDateTime::UtcNow().ToUnixTimestamp() + 90, "join", FMath::Rand(), nullptr, LoginSession.UserUri.ToString(), ChannelUri.ToString(), (const unsigned char*)TCHAR_TO_ANSI(*VivoxInsecureSecret), VivoxInsecureSecret.Len()))
+		VivoxClientApi::Uri ChannelUri = VivoxVoiceChat.CreateChannelUri(ChannelName, ChannelType, Channel3dProperties);
+		if (char* ANSIToken = vx_debug_generate_token(TCHAR_TO_ANSI(*VivoxVoiceChat.VivoxIssuer), FDateTime::UtcNow().ToUnixTimestamp() + 90, "join", FMath::Rand(), nullptr, LoginSession.UserUri.ToString(), ChannelUri.ToString(), (const unsigned char*)TCHAR_TO_ANSI(*VivoxVoiceChat.VivoxInsecureSecret), VivoxVoiceChat.VivoxInsecureSecret.Len()))
 		{
 			Token = ANSI_TO_TCHAR(ANSIToken);
 			vx_free(ANSIToken);
@@ -1185,19 +960,6 @@
 	return Token;
 }
 
-void FVivoxVoiceChat::InvokeOnUIThread(void (Func)(void* Arg0), void* Arg0)
-{
-	AsyncTask(ENamedThreads::GameThread, [Func, Arg0]()
-	{
-		if (Func)
-		{
-			(*Func)(Arg0);
-		}
-	});
-
-	FEmbeddedCommunication::WakeGameThread();
-}
-
 FString GetLogLevelString(VivoxClientApi::IClientApiEventHandler::LogLevel Level)
 {
 	switch (Level)
@@ -1211,102 +973,384 @@
 	}
 }
 
-void FVivoxVoiceChat::onLogStatementEmitted(LogLevel Level, long long NativeMillisecondsSinceEpoch, long ThreadId, const char* LogMessageCStr)
+bool FVivoxVoiceChatUser::IsInitialized()
 {
-	FString LogMessage = ANSI_TO_TCHAR(LogMessageCStr);
-	bool bLogMatches = false;
+	return VivoxVoiceChat.IsInitialized();
+}
 
-	// Don't spam the log if we receive repeated log messages
-	if (!LastLogMessage.IsEmpty())
+bool FVivoxVoiceChatUser::Exec(UWorld* InWorld, const TCHAR* Cmd, FOutputDevice& Ar)
+{
+#define VIVOX_EXEC_LOG(Fmt, ...) Ar.CategorizedLogf(LogVivoxVoiceChat.GetCategoryName(), ELogVerbosity::Log, Fmt, ##__VA_ARGS__)
+
+	if (FParse::Command(&Cmd, TEXT("INFO")))
 	{
-		bLogMatches = LogMessage == LastLogMessage && Level == LastLogLevel;
-		if (bLogMatches)
+		if (IsInitialized())
 		{
-			LogSpamCount++;
-		}
-		
-		const bool bLogSpamCountPowerOfTwo = FMath::IsPowerOfTwo(LogSpamCount);
-		const bool bLogDueToPot = bLogMatches && bLogSpamCountPowerOfTwo;
-		const bool bLogDueToChange = !bLogMatches && LogSpamCount > 0 && !bLogSpamCountPowerOfTwo; // Don't log on change if POT, as we already logged it.
-		if(bLogDueToChange || bLogDueToPot)
-		{
-			if (LastLogLevel == LogLevelError)
+			VIVOX_EXEC_LOG(TEXT("  Input Devices: muted:%s volume:%.2f"), *LexToString(GetAudioInputDeviceMuted()), GetAudioInputVolume());
+			const FString InputDevice = GetInputDevice();
+			const FString DefaultInputDevice = GetDefaultInputDevice();
+			if (InputDevice == DefaultInputDevice)
 			{
-				UE_LOG(LogVivoxVoiceChat, Warning, TEXT("vivox: Error: %s (seen %d times)"), *LastLogMessage, LogSpamCount);
+				VIVOX_EXEC_LOG(TEXT("    [%s] (Selected) (Default)"), *DefaultInputDevice);
 			}
 			else
 			{
-				UE_LOG(LogVivoxVoiceChat, Log, TEXT("vivox: %s: %s (seen %d times)"), *GetLogLevelString(LastLogLevel), *LastLogMessage, LogSpamCount);
+				VIVOX_EXEC_LOG(TEXT("    [%s] (Selected)"), *InputDevice);
+				VIVOX_EXEC_LOG(TEXT("    [%s] (Default)"), *DefaultInputDevice);
+			}
+			for (const FString& Device : GetAvailableInputDevices())
+			{
+				if (Device != DefaultInputDevice && Device != InputDevice)
+				{
+					VIVOX_EXEC_LOG(TEXT("    [%s]"), *Device);
+				}
+			}
+
+			VIVOX_EXEC_LOG(TEXT("  Output Devices: muted:%s volume:%.2f"), *LexToString(GetAudioOutputDeviceMuted()), GetAudioOutputVolume());
+			const FString OutputDevice = GetOutputDevice();
+			const FString DefaultOutputDevice = GetDefaultOutputDevice();
+			if (OutputDevice == DefaultOutputDevice)
+			{
+				VIVOX_EXEC_LOG(TEXT("    [%s] (Selected) (Default)"), *DefaultOutputDevice);
+			}
+			else
+			{
+				VIVOX_EXEC_LOG(TEXT("    [%s] (Selected)"), *OutputDevice);
+				VIVOX_EXEC_LOG(TEXT("    [%s] (Default)"), *DefaultOutputDevice);
+			}
+			for (const FString& Device : GetAvailableOutputDevices())
+			{
+				if (Device != DefaultOutputDevice && Device != OutputDevice)
+				{
+					VIVOX_EXEC_LOG(TEXT("    [%s]"), *Device);
+				}
+			}
+
+			if (IsConnected())
+			{
+				VIVOX_EXEC_LOG(TEXT("Login Status: %s"), *ToString(LoginSession.State));
+				if (IsLoggedIn())
+				{
+					VIVOX_EXEC_LOG(TEXT("  PlayerName: %s"), *LoginSession.PlayerName);
+					VIVOX_EXEC_LOG(TEXT("  AccountName: %s"), ANSI_TO_TCHAR(LoginSession.AccountName.ToString()));
+					VIVOX_EXEC_LOG(TEXT("  UserUri: %s"), ANSI_TO_TCHAR(LoginSession.UserUri.ToString()));
+
+					VivoxClientApi::ChannelTransmissionPolicy TransmissionPolicy = VivoxClientConnection.GetChannelTransmissionPolicy(LoginSession.AccountName);
+					FString TransmitString;
+					switch (GetTransmitMode())
+					{
+					case EVoiceChatTransmitMode::All:
+						TransmitString = TEXT("ALL");
+						break;
+					case EVoiceChatTransmitMode::None:
+						TransmitString = TEXT("NONE");
+						break;
+					case EVoiceChatTransmitMode::Channel:
+						TransmitString = FString::Printf(TEXT("CHANNEL:%s"), *GetTransmitChannel());
+						break;
+					}
+					VIVOX_EXEC_LOG(TEXT("Channels: transmitting:%s"), *TransmitString);
+					for (const TPair<FString, FChannelSession>& ChannelSessionPair : LoginSession.ChannelSessions)
+					{
+						const FString& ChannelName = ChannelSessionPair.Key;
+						const FChannelSession& ChannelSession = ChannelSessionPair.Value;
+						VIVOX_EXEC_LOG(TEXT("  %s"), *ChannelName);
+						VIVOX_EXEC_LOG(TEXT("    Channel Status: %s"), *ToString(ChannelSession.State));
+						VIVOX_EXEC_LOG(TEXT("    Channel Uri: %s"), ANSI_TO_TCHAR(ChannelSession.ChannelUri.ToString()));
+						VIVOX_EXEC_LOG(TEXT("    Participants:"));
+						for (const TPair<FString, FParticipant>& ParticipantPair : ChannelSession.Participants)
+						{
+							const FString& ParticipantName = ParticipantPair.Key;
+							const FParticipant& Participant = ParticipantPair.Value;
+							VIVOX_EXEC_LOG(TEXT("      %s uri:%s talking:%s muted:%s volume:%.2f"), *ParticipantName, ANSI_TO_TCHAR(Participant.UserUri.ToString()), *LexToString(Participant.bTalking), *LexToString(Participant.bMuted), Participant.Volume);
+						}
+					}
+				}
+			}
+		}
+		return true;
+	}
+#if !UE_BUILD_SHIPPING
+	else if (FParse::Command(&Cmd, TEXT("INPUT")))
+	{
+		if (FParse::Command(&Cmd, TEXT("SETVOLUME")))
+		{
+			FString Volume;
+			if (FParse::Token(Cmd, Volume, false))
+			{
+				SetAudioInputVolume(FCString::Atof(*Volume));
+				return true;
+			}
+		}
+		else if (FParse::Command(&Cmd, TEXT("MUTE")))
+		{
+			SetAudioInputDeviceMuted(true);
+			return true;
+		}
+		else if (FParse::Command(&Cmd, TEXT("UNMUTE")))
+		{
+			SetAudioInputDeviceMuted(false);
+			return true;
+		}
+		else if (FParse::Command(&Cmd, TEXT("LISTDEVICES")))
+		{
+			TArray<FString> Devices = GetAvailableInputDevices();
+			VIVOX_EXEC_LOG(TEXT("Input Devices:"));
+			for (const FString& Device : Devices)
+			{
+				VIVOX_EXEC_LOG(TEXT("  %s"), *Device);
+			}
+			return true;
+		}
+		else if (FParse::Command(&Cmd, TEXT("SETDEVICE")))
+		{
+			FString Device;
+			if (FParse::Token(Cmd, Device, false))
+			{
+				SetInputDevice(Device);
+				return true;
+			}
+		}
+		else if (FParse::Command(&Cmd, TEXT("SETDEFAULTDEVICE")))
+		{
+			SetInputDevice(FString());
+			return true;
+		}
+	}
+	else if (FParse::Command(&Cmd, TEXT("OUTPUT")))
+	{
+		if (FParse::Command(&Cmd, TEXT("SETVOLUME")))
+		{
+			FString Volume;
+			if (FParse::Token(Cmd, Volume, false))
+			{
+				SetAudioOutputVolume(FCString::Atof(*Volume));
+				return true;
+			}
+		}
+		else if (FParse::Command(&Cmd, TEXT("MUTE")))
+		{
+			SetAudioOutputDeviceMuted(true);
+			return true;
+		}
+		else if (FParse::Command(&Cmd, TEXT("UNMUTE")))
+		{
+			SetAudioOutputDeviceMuted(false);
+			return true;
+		}
+		else if (FParse::Command(&Cmd, TEXT("LISTDEVICES")))
+		{
+			TArray<FString> Devices = GetAvailableOutputDevices();
+			VIVOX_EXEC_LOG(TEXT("Output Devices:"));
+			for (const FString& Device : Devices)
+			{
+				VIVOX_EXEC_LOG(TEXT("  %s"), *Device);
+			}
+			return true;
+		}
+		else if (FParse::Command(&Cmd, TEXT("SETDEVICE")))
+		{
+			FString Device;
+			if (FParse::Token(Cmd, Device, false))
+			{
+				SetOutputDevice(Device);
+				return true;
+			}
+		}
+		else if (FParse::Command(&Cmd, TEXT("SETDEFAULTDEVICE")))
+		{
+			SetOutputDevice(FString());
+			return true;
+		}
+	}
+	else if (FParse::Command(&Cmd, TEXT("CONNECT")))
+	{
+		Connect(FOnVoiceChatConnectCompleteDelegate::CreateLambda([](const FVoiceChatResult& Result)
+		{
+			UE_LOG(LogVivoxVoiceChat, Log, TEXT("VIVOX CONNECT result:%s"), *LexToString(Result));
+		}));
+		return true;
+	}
+	else if (FParse::Command(&Cmd, TEXT("DISCONNECT")))
+	{
+		Disconnect(FOnVoiceChatDisconnectCompleteDelegate::CreateLambda([](const FVoiceChatResult& Result)
+		{
+			UE_LOG(LogVivoxVoiceChat, Log, TEXT("VIVOX DISCONNECT result:%s"), *LexToString(Result));
+		}));
+		return true;
+	}
+	else if (FParse::Command(&Cmd, TEXT("LOGIN")))
+	{
+		FString PlayerName;
+		if (FParse::Token(Cmd, PlayerName, false))
+		{
+			FString Token = InsecureGetLoginToken(PlayerName);
+
+			Login(0, PlayerName, Token, FOnVoiceChatLoginCompleteDelegate::CreateLambda([](const FString& LoggedInPlayerName, const FVoiceChatResult& Result)
+			{
+				UE_LOG(LogVivoxVoiceChat, Log, TEXT("VIVOX LOGIN playername:%s result:%s"), *LoggedInPlayerName, *LexToString(Result));
+			}));
+			return true;
+		}
+	}
+	else if (FParse::Command(&Cmd, TEXT("LOGOUT")))
+	{
+		Logout(FOnVoiceChatLogoutCompleteDelegate::CreateLambda([](const FString& PlayerName, const FVoiceChatResult& Result)
+		{
+			UE_LOG(LogVivoxVoiceChat, Log, TEXT("VIVOX LOGOUT playername:%s result:%s"), *PlayerName, *LexToString(Result));
+		}));
+		return true;
+	}
+	else if (FParse::Command(&Cmd, TEXT("CHANNEL")))
+	{
+		if (FParse::Command(&Cmd, TEXT("JOIN")))
+		{
+			FString ChannelName;
+			if (FParse::Token(Cmd, ChannelName, false))
+			{
+				FString ChannelTypeString;
+				EVoiceChatChannelType ChannelType = EVoiceChatChannelType::NonPositional;
+				TOptional<FVoiceChatChannel3dProperties> Channel3dProperties;
+				if (FParse::Token(Cmd, ChannelTypeString, false))
+				{
+					if (ChannelTypeString == TEXT("POSITIONAL"))
+					{
+						ChannelType = EVoiceChatChannelType::Positional;
+					}
+					else if (ChannelTypeString == TEXT("ECHO"))
+					{
+						ChannelType = EVoiceChatChannelType::Echo;
+					}
+				}
+
+				FString Token = InsecureGetJoinToken(ChannelName, ChannelType, Channel3dProperties);
+
+				JoinChannel(ChannelName, Token, ChannelType, FOnVoiceChatChannelJoinCompleteDelegate::CreateLambda([](const FString& JoinedChannelName, const FVoiceChatResult& Result)
+				{
+					UE_LOG(LogVivoxVoiceChat, Log, TEXT("VIVOX CHANNEL JOIN channelname:%s result:%s"), *JoinedChannelName, *LexToString(Result));
+				}), Channel3dProperties);
+				return true;
+			}
+		}
+		else if (FParse::Command(&Cmd, TEXT("LEAVE")))
+		{
+			FString ChannelName;
+			if (FParse::Token(Cmd, ChannelName, false))
+			{
+				LeaveChannel(ChannelName, FOnVoiceChatChannelLeaveCompleteDelegate::CreateLambda([](const FString& LeftChannelName, const FVoiceChatResult& Result)
+				{
+					UE_LOG(LogVivoxVoiceChat, Log, TEXT("VIVOX CHANNEL LEAVE channelname:%s result:%s"), *LeftChannelName, *LexToString(Result));
+				}));
+				return true;
+			}
+		}
+		else if (FParse::Command(&Cmd, TEXT("TRANSMIT")))
+		{
+			FString ChannelName;
+			if (FParse::Token(Cmd, ChannelName, false))
+			{
+				TransmitToSpecificChannel(ChannelName);
+				return true;
+			}
+		}
+		else if (FParse::Command(&Cmd, TEXT("TRANSMITALL")))
+		{
+			TransmitToAllChannels();
+			return true;
+		}
+		else if (FParse::Command(&Cmd, TEXT("TRANSMITNONE")))
+		{
+			TransmitToNoChannels();
+			return true;
+		}
+	}
+	else if (FParse::Command(&Cmd, TEXT("PLAYER")))
+	{
+		if (FParse::Command(&Cmd, TEXT("MUTE")))
+		{
+			FString PlayerName;
+			if (FParse::Token(Cmd, PlayerName, false))
+			{
+				SetPlayerMuted(PlayerName, true);
+				return true;
+			}
+		}
+		else if (FParse::Command(&Cmd, TEXT("UNMUTE")))
+		{
+			FString PlayerName;
+			if (FParse::Token(Cmd, PlayerName, false))
+			{
+				SetPlayerMuted(PlayerName, false);
+				return true;
+			}
+		}
+		else if (FParse::Command(&Cmd, TEXT("SETVOLUME")))
+		{
+			FString PlayerName;
+			if (FParse::Token(Cmd, PlayerName, false))
+			{
+				FString Volume;
+				if (FParse::Token(Cmd, Volume, false))
+				{
+					SetPlayerVolume(PlayerName, FCString::Atof(*Volume));
+					return true;
+				}
+			}
+		}
+		else if (FParse::Command(&Cmd, TEXT("BLOCK")))
+		{
+			TArray<FString> PlayerNames;
+			FString PlayerName;
+			while (FParse::Token(Cmd, PlayerName, false))
+			{
+				PlayerNames.Add(PlayerName);
+			}
+			if (PlayerNames.Num() > 0)
+			{
+				BlockPlayers(PlayerNames);
+				return true;
+			}
+		}
+		else if (FParse::Command(&Cmd, TEXT("UNBLOCK")))
+		{
+			TArray<FString> PlayerNames;
+			FString PlayerName;
+			while (FParse::Token(Cmd, PlayerName, false))
+			{
+				PlayerNames.Add(PlayerName);
+			}
+			if (PlayerNames.Num() > 0)
+			{
+				UnblockPlayers(PlayerNames);
+				return true;
 			}
 		}
 	}
+#endif // !UE_BUILD_SHIPPING
 
-	if (!bLogMatches)
-	{
-		if (Level == LogLevelError)
-		{
-			UE_LOG(LogVivoxVoiceChat, Warning, TEXT("vivox: Error: %s"), *LogMessage);
-		}
-		else
-		{
-			UE_LOG(LogVivoxVoiceChat, Log, TEXT("vivox: %s: %s"), *GetLogLevelString(Level), *LogMessage);
-		}
+#undef VIVOX_EXEC_LOG
 
-		LogSpamCount = 0;
-		LastLogMessage = MoveTemp(LogMessage);
-		LastLogLevel = Level;
-	}
+	return false;
 }
 
-void FVivoxVoiceChat::onConnectCompleted(const VivoxClientApi::Uri& Server)
+void FVivoxVoiceChatUser::onConnectCompleted(const VivoxClientApi::Uri& Server)
 {
-	UE_LOG(LogVivoxVoiceChat, Log, TEXT("onConnectCompleted server:%s"), ANSI_TO_TCHAR(Server.ToString()));
-
-	ConnectionState = EConnectionState::Connected;
-
-	TriggerCompletionDelegates(OnVoiceChatConnectCompleteDelegates, FVoiceChatResult::CreateSuccess());
 }
 
-void FVivoxVoiceChat::onConnectFailed(const VivoxClientApi::Uri& Server, const VivoxClientApi::VCSStatus& Status)
+void FVivoxVoiceChatUser::onConnectFailed(const VivoxClientApi::Uri& Server, const VivoxClientApi::VCSStatus& Status)
 {
-	UE_LOG(LogVivoxVoiceChat, Warning, TEXT("onConnectFailed server:%s error:%s (%i)"), ANSI_TO_TCHAR(Server.ToString()), ANSI_TO_TCHAR(Status.ToString()), Status.GetStatusCode());
-
-	ConnectionState = EConnectionState::Disconnected;
-
-	TriggerCompletionDelegates(OnVoiceChatConnectCompleteDelegates, ResultFromVivoxStatus(Status));
 }
 
-void FVivoxVoiceChat::onDisconnected(const VivoxClientApi::Uri& Server, const VivoxClientApi::VCSStatus& Status)
+void FVivoxVoiceChatUser::onDisconnected(const VivoxClientApi::Uri& Server, const VivoxClientApi::VCSStatus& Status)
 {
-	if (Status.IsError())
-	{
-		UE_LOG(LogVivoxVoiceChat, Warning, TEXT("onDisconnected server:%s error:%s (%i)"), ANSI_TO_TCHAR(Server.ToString()), ANSI_TO_TCHAR(Status.ToString()), Status.GetStatusCode());
-	}
-	else
-	{
-		UE_LOG(LogVivoxVoiceChat, Log, TEXT("onDisconnected server:%s"), ANSI_TO_TCHAR(Server.ToString()));
-	}
-
 	ClearLoginSession();
-
-	EConnectionState PreviousConnectionState = ConnectionState;
-	ConnectionState = EConnectionState::Disconnected;
-
-	if (PreviousConnectionState == EConnectionState::Disconnecting)
-	{
-		TriggerCompletionDelegates(OnVoiceChatDisconnectCompleteDelegates, FVoiceChatResult::CreateSuccess());
-	}
-	else
-	{
-		OnVoiceChatDisconnectedDelegate.Broadcast(ResultFromVivoxStatus(Status));
-	}
 }
 
-void FVivoxVoiceChat::onLoginCompleted(const VivoxClientApi::AccountName& AccountName)
+void FVivoxVoiceChatUser::onLoginCompleted(const VivoxClientApi::AccountName& AccountName)
 {
 	UE_LOG(LogVivoxVoiceChat, Log, TEXT("onLoginCompleted account:%s"), ANSI_TO_TCHAR(AccountName.ToString()));
 
-	FString PlayerName = GetPlayerNameFromAccountName(AccountName);
+	FString PlayerName = VivoxVoiceChat.GetPlayerNameFromAccountName(AccountName);
 
 	LoginSession.State = FLoginSession::EState::LoggedIn;
 
@@ -1315,29 +1359,29 @@
 	OnVoiceChatLoggedInDelegate.Broadcast(PlayerName);
 }
 
-void FVivoxVoiceChat::onInvalidLoginCredentials(const VivoxClientApi::AccountName& AccountName)
+void FVivoxVoiceChatUser::onInvalidLoginCredentials(const VivoxClientApi::AccountName& AccountName)
 {
 	UE_LOG(LogVivoxVoiceChat, Warning, TEXT("onInvalidLoginCredentials account:%s"), ANSI_TO_TCHAR(AccountName.ToString()));
 
 	LoginSession.State = FLoginSession::EState::LoggedOut;
 
-	FString PlayerName = GetPlayerNameFromAccountName(AccountName);
+	FString PlayerName = VivoxVoiceChat.GetPlayerNameFromAccountName(AccountName);
 
 	TriggerCompletionDelegate(OnVoiceChatLoginCompleteDelegate, PlayerName, VoiceChat::Errors::CredentialsInvalid());
 }
 
-void FVivoxVoiceChat::onLoginFailed(const VivoxClientApi::AccountName& AccountName, const VivoxClientApi::VCSStatus& Status)
+void FVivoxVoiceChatUser::onLoginFailed(const VivoxClientApi::AccountName& AccountName, const VivoxClientApi::VCSStatus& Status)
 {
 	UE_LOG(LogVivoxVoiceChat, Warning, TEXT("onLoginFailed account:%s error:%s (%i)"), ANSI_TO_TCHAR(AccountName.ToString()), ANSI_TO_TCHAR(Status.ToString()), Status.GetStatusCode());
 
 	LoginSession.State = FLoginSession::EState::LoggedOut;
 
-	FString PlayerName = GetPlayerNameFromAccountName(AccountName);
+	FString PlayerName = VivoxVoiceChat.GetPlayerNameFromAccountName(AccountName);
 
 	TriggerCompletionDelegate(OnVoiceChatLoginCompleteDelegate, PlayerName, ResultFromVivoxStatus(Status));
 }
 
-void FVivoxVoiceChat::onLogoutCompleted(const VivoxClientApi::AccountName& AccountName)
+void FVivoxVoiceChatUser::onLogoutCompleted(const VivoxClientApi::AccountName& AccountName)
 {
 	UE_LOG(LogVivoxVoiceChat, Log, TEXT("onLogoutCompleted account:%s"), ANSI_TO_TCHAR(AccountName.ToString()));
 
@@ -1345,23 +1389,28 @@
 
 	LoginSession.State = FLoginSession::EState::LoggedOut;
 
-	FString PlayerName = GetPlayerNameFromAccountName(AccountName);
+	FString PlayerName = VivoxVoiceChat.GetPlayerNameFromAccountName(AccountName);
 	TriggerCompletionDelegate(OnVoiceChatLogoutCompleteDelegate, PlayerName, FVoiceChatResult::CreateSuccess());
 
 	OnVoiceChatLoggedOutDelegate.Broadcast(PlayerName);
 }
 
-void FVivoxVoiceChat::onLogoutFailed(const VivoxClientApi::AccountName& AccountName, const VivoxClientApi::VCSStatus& Status)
+void FVivoxVoiceChatUser::onLogoutFailed(const VivoxClientApi::AccountName& AccountName, const VivoxClientApi::VCSStatus& Status)
 {
 	UE_LOG(LogVivoxVoiceChat, Warning, TEXT("onLogoutFailed account:%s error:%s (%i)"), ANSI_TO_TCHAR(AccountName.ToString()), ANSI_TO_TCHAR(Status.ToString()), Status.GetStatusCode());
 
 	LoginSession.State = FLoginSession::EState::LoggedIn;
 
-	FString PlayerName = GetPlayerNameFromAccountName(AccountName);
+	FString PlayerName = VivoxVoiceChat.GetPlayerNameFromAccountName(AccountName);
 	TriggerCompletionDelegate(OnVoiceChatLogoutCompleteDelegate, PlayerName, ResultFromVivoxStatus(Status));
 }
 
-void FVivoxVoiceChat::onChannelJoined(const VivoxClientApi::AccountName& AccountName, const VivoxClientApi::Uri& ChannelUri)
+void FVivoxVoiceChatUser::onSessionGroupCreated(const VivoxClientApi::AccountName& AccountName, const char* SessionGroupHandle)
 {
+	SessionGroup = ANSI_TO_TCHAR(SessionGroupHandle);
+}
+
+void FVivoxVoiceChatUser::onChannelJoined(const VivoxClientApi::AccountName& AccountName, const VivoxClientApi::Uri& ChannelUri)
+{
 	UE_LOG(LogVivoxVoiceChat, Log, TEXT("onChannelJoined channel:%s"), ANSI_TO_TCHAR(ChannelUri.ToString()));
 
@@ -1372,7 +1421,7 @@
 	OnVoiceChatChannelJoinedDelegate.Broadcast(ChannelSession.ChannelName);
 }
 
-void FVivoxVoiceChat::onInvalidChannelCredentials(const VivoxClientApi::AccountName& AccountName, const VivoxClientApi::Uri& ChannelUri)
+void FVivoxVoiceChatUser::onInvalidChannelCredentials(const VivoxClientApi::AccountName& AccountName, const VivoxClientApi::Uri& ChannelUri)
 {
 	UE_LOG(LogVivoxVoiceChat, Warning, TEXT("onInvalidChannelCredentials channel:%s"), ANSI_TO_TCHAR(ChannelUri.ToString()));
 
@@ -1383,7 +1432,7 @@
 	RemoveChannelSession(ChannelSession.ChannelName);
 }
 
-void FVivoxVoiceChat::onChannelJoinFailed(const VivoxClientApi::AccountName& AccountName, const VivoxClientApi::Uri& ChannelUri, const VivoxClientApi::VCSStatus& Status)
+void FVivoxVoiceChatUser::onChannelJoinFailed(const VivoxClientApi::AccountName& AccountName, const VivoxClientApi::Uri& ChannelUri, const VivoxClientApi::VCSStatus& Status)
 {
 	UE_LOG(LogVivoxVoiceChat, Warning, TEXT("onChannelJoinFailed channel:%s error:%s (%i)"), ANSI_TO_TCHAR(ChannelUri.ToString()), ANSI_TO_TCHAR(Status.ToString()), Status.GetStatusCode());
 
@@ -1394,7 +1443,7 @@
 	RemoveChannelSession(ChannelSession.ChannelName);
 }
 
-void FVivoxVoiceChat::onChannelExited(const VivoxClientApi::AccountName& AccountName, const VivoxClientApi::Uri& ChannelUri, const VivoxClientApi::VCSStatus& Status)
+void FVivoxVoiceChatUser::onChannelExited(const VivoxClientApi::AccountName& AccountName, const VivoxClientApi::Uri& ChannelUri, const VivoxClientApi::VCSStatus& Status)
 {
 	const FVoiceChatResult Result = ResultFromVivoxStatus(Status);
 	if (Result.IsSuccess())
@@ -1429,7 +1478,7 @@
 	RemoveChannelSession(ChannelSession.ChannelName);
 }
 
-void FVivoxVoiceChat::onCallStatsUpdated(const VivoxClientApi::AccountName& AccountName, vx_call_stats_t& Stats, bool bIsFinal)
+void FVivoxVoiceChatUser::onCallStatsUpdated(const VivoxClientApi::AccountName& AccountName, vx_call_stats_t& Stats, bool bIsFinal)
 {
 	const int TotalPacketsLost = Stats.incoming_packetloss + Stats.incoming_discarded + Stats.incoming_out_of_time;
 	const int TotalPackets = TotalPacketsLost + Stats.incoming_received;
@@ -1445,12 +1494,12 @@
 	OnVoiceChatCallStatsUpdatedDelegate.Broadcast(CallStats);
 }
 
-void FVivoxVoiceChat::onParticipantAdded(const VivoxClientApi::AccountName& AccountName, const VivoxClientApi::Uri& ChannelUri, const VivoxClientApi::Uri& ParticipantUri, bool bIsLoggedInUser)
+void FVivoxVoiceChatUser::onParticipantAdded(const VivoxClientApi::AccountName& AccountName, const VivoxClientApi::Uri& ChannelUri, const VivoxClientApi::Uri& ParticipantUri, bool bIsLoggedInUser)
 {
 	UE_LOG(LogVivoxVoiceChat, Log, TEXT("onParticipantAdded channel:%s participant:%s"), ANSI_TO_TCHAR(ChannelUri.ToString()), ANSI_TO_TCHAR(ParticipantUri.ToString()));
 
 	FChannelSession& ChannelSession = GetChannelSession(ChannelUri);
-	FString PlayerName = GetPlayerNameFromUri(ParticipantUri);
+	FString PlayerName = VivoxVoiceChat.GetPlayerNameFromUri(ParticipantUri);
 	FParticipant& ChannelParticipant = ChannelSession.Participants.Add(PlayerName);
 	ChannelParticipant.PlayerName = PlayerName;
 	ChannelParticipant.UserUri = ParticipantUri;
@@ -1491,24 +1540,24 @@
 	}
 }
 
-void FVivoxVoiceChat::onParticipantLeft(const VivoxClientApi::AccountName& AccountName, const VivoxClientApi::Uri& ChannelUri, const VivoxClientApi::Uri& ParticipantUri, bool bIsLoggedInUser, ParticipantLeftReason Reason)
+void FVivoxVoiceChatUser::onParticipantLeft(const VivoxClientApi::AccountName& AccountName, const VivoxClientApi::Uri& ChannelUri, const VivoxClientApi::Uri& ParticipantUri, bool bIsLoggedInUser, VivoxClientApi::IClientApiEventHandler::ParticipantLeftReason Reason)
 {
 	UE_LOG(LogVivoxVoiceChat, Log, TEXT("onParticipantLeft channel:%s participant:%s reason:%s"), ANSI_TO_TCHAR(ChannelUri.ToString()), ANSI_TO_TCHAR(ParticipantUri.ToString()), *LexToString(Reason));
 
 	FChannelSession& ChannelSession = GetChannelSession(ChannelUri);
-	FString PlayerName = GetPlayerNameFromUri(ParticipantUri);
+	FString PlayerName = VivoxVoiceChat.GetPlayerNameFromUri(ParticipantUri);
 
 	ChannelSession.Participants.Remove(PlayerName);
 
 	OnVoiceChatPlayerRemovedDelegate.Broadcast(ChannelSession.ChannelName, PlayerName);
 }
 
-void FVivoxVoiceChat::onParticipantUpdated(const VivoxClientApi::AccountName& AccountName, const VivoxClientApi::Uri& ChannelUri, const VivoxClientApi::Uri& ParticipantUri, bool bIsLoggedInUser, bool bSpeaking, double MeterEnergy, bool bMutedForAll)
+void FVivoxVoiceChatUser::onParticipantUpdated(const VivoxClientApi::AccountName& AccountName, const VivoxClientApi::Uri& ChannelUri, const VivoxClientApi::Uri& ParticipantUri, bool bIsLoggedInUser, bool bSpeaking, double MeterEnergy, bool bMutedForAll)
 {
 	UE_LOG(LogVivoxVoiceChat, VeryVerbose, TEXT("onParticipantUpdated channel:%s participant:%s speaking:%s energy:%f"), ANSI_TO_TCHAR(ChannelUri.ToString()), ANSI_TO_TCHAR(ParticipantUri.ToString()), *LexToString(bSpeaking), MeterEnergy);
 
 	FChannelSession& ChannelSession = GetChannelSession(ChannelUri);
-	FString PlayerName = GetPlayerNameFromUri(ParticipantUri);
+	FString PlayerName = VivoxVoiceChat.GetPlayerNameFromUri(ParticipantUri);
 
 	FParticipant& Participant = GetParticipant(PlayerName);
 	Participant.bTalking = bSpeaking;
@@ -1520,49 +1569,49 @@
 	}
 }
 
-void FVivoxVoiceChat::onAvailableAudioDevicesChanged()
+void FVivoxVoiceChatUser::onAvailableAudioDevicesChanged()
 {
 	UE_LOG(LogVivoxVoiceChat, Verbose, TEXT("onAvailableAudioDevicesChanged"));
 
 	OnVoiceChatAvailableAudioDevicesChangedDelegate.Broadcast();
 }
 
-void FVivoxVoiceChat::onOperatingSystemChosenAudioInputDeviceChanged(const VivoxClientApi::AudioDeviceId& DeviceId)
+void FVivoxVoiceChatUser::onOperatingSystemChosenAudioInputDeviceChanged(const VivoxClientApi::AudioDeviceId& DeviceId)
 {
 	UE_LOG(LogVivoxVoiceChat, Verbose, TEXT("onOperatingSystemChosenAudioInputDeviceChanged deviceid:%s"), UTF8_TO_TCHAR(DeviceId.ToString()));
 }
 
-void FVivoxVoiceChat::onSetApplicationChosenAudioInputDeviceCompleted(const VivoxClientApi::AudioDeviceId& DeviceId)
+void FVivoxVoiceChatUser::onSetApplicationChosenAudioInputDeviceCompleted(const VivoxClientApi::AccountName& AccountName, const VivoxClientApi::AudioDeviceId& DeviceId)
 {
 	UE_LOG(LogVivoxVoiceChat, Verbose, TEXT("onSetApplicationChosenAudioInputDeviceCompleted deviceid:%s"), UTF8_TO_TCHAR(DeviceId.ToString()));
 }
 
-void FVivoxVoiceChat::onSetApplicationChosenAudioInputDeviceFailed(const VivoxClientApi::AudioDeviceId& DeviceId, const VivoxClientApi::VCSStatus& Status)
+void FVivoxVoiceChatUser::onSetApplicationChosenAudioInputDeviceFailed(const VivoxClientApi::AccountName& AccountName, const VivoxClientApi::AudioDeviceId& DeviceId, const VivoxClientApi::VCSStatus& Status)
 {
 	UE_LOG(LogVivoxVoiceChat, Warning, TEXT("onSetApplicationChosenAudioInputDeviceFailed deviceid:%s error:%s (%i)"), UTF8_TO_TCHAR(DeviceId.ToString()), ANSI_TO_TCHAR(Status.ToString()), Status.GetStatusCode());
 }
 
-void FVivoxVoiceChat::onOperatingSystemChosenAudioOutputDeviceChanged(const VivoxClientApi::AudioDeviceId& DeviceId)
+void FVivoxVoiceChatUser::onOperatingSystemChosenAudioOutputDeviceChanged(const VivoxClientApi::AudioDeviceId& DeviceId)
 {
 	UE_LOG(LogVivoxVoiceChat, Verbose, TEXT("onOperatingSystemChosenAudioOutputDeviceChanged deviceid:%s"), UTF8_TO_TCHAR(DeviceId.ToString()));
 }
 
-void FVivoxVoiceChat::onSetApplicationChosenAudioOutputDeviceCompleted(const VivoxClientApi::AudioDeviceId& DeviceId)
+void FVivoxVoiceChatUser::onSetApplicationChosenAudioOutputDeviceCompleted(const VivoxClientApi::AccountName& AccountName, const VivoxClientApi::AudioDeviceId& DeviceId)
 {
 	UE_LOG(LogVivoxVoiceChat, Verbose, TEXT("onSetApplicationChosenAudioOutputDeviceCompleted deviceid:%s"), UTF8_TO_TCHAR(DeviceId.ToString()));
 }
 
-void FVivoxVoiceChat::onSetApplicationChosenAudioOutputDeviceFailed(const VivoxClientApi::AudioDeviceId& DeviceId, const VivoxClientApi::VCSStatus& Status)
+void FVivoxVoiceChatUser::onSetApplicationChosenAudioOutputDeviceFailed(const VivoxClientApi::AccountName& AccountName, const VivoxClientApi::AudioDeviceId& DeviceId, const VivoxClientApi::VCSStatus& Status)
 {
 	UE_LOG(LogVivoxVoiceChat, Warning, TEXT("onSetApplicationChosenAudioOutputDeviceFailed deviceid:%s error:%s (%i)"), UTF8_TO_TCHAR(DeviceId.ToString()), ANSI_TO_TCHAR(Status.ToString()), Status.GetStatusCode());
 }
 
-void FVivoxVoiceChat::onSetParticipantAudioOutputDeviceVolumeForMeCompleted(const VivoxClientApi::AccountName& AccountName, const VivoxClientApi::Uri& TargetUser, const VivoxClientApi::Uri& ChannelUri, int Volume)
+void FVivoxVoiceChatUser::onSetParticipantAudioOutputDeviceVolumeForMeCompleted(const VivoxClientApi::AccountName& AccountName, const VivoxClientApi::Uri& TargetUser, const VivoxClientApi::Uri& ChannelUri, int Volume)
 {
 	UE_LOG(LogVivoxVoiceChat, Verbose, TEXT("onSetParticipantAudioOutputDeviceVolumeForMeCompleted channel:%s user:%s volume:%i"), ANSI_TO_TCHAR(ChannelUri.ToString()), ANSI_TO_TCHAR(TargetUser.ToString()), Volume);
 
 	FChannelSession& ChannelSession = GetChannelSession(ChannelUri);
-	FString PlayerName = GetPlayerNameFromUri(TargetUser);
+	FString PlayerName = VivoxVoiceChat.GetPlayerNameFromUri(TargetUser);
 	if (FParticipant* ChannelParticipant = ChannelSession.Participants.Find(PlayerName))
 	{
 		ChannelParticipant->Volume = static_cast<float>(Volume - VIVOX_MIN_VOL) / static_cast<float>(VIVOX_MAX_VOL - VIVOX_MIN_VOL);
@@ -1571,34 +1620,34 @@
 	}
 }
 
-void FVivoxVoiceChat::onSetParticipantAudioOutputDeviceVolumeForMeFailed(const VivoxClientApi::AccountName& AccountName, const VivoxClientApi::Uri& TargetUser, const VivoxClientApi::Uri& ChannelUri, int Volume, const VivoxClientApi::VCSStatus& Status)
+void FVivoxVoiceChatUser::onSetParticipantAudioOutputDeviceVolumeForMeFailed(const VivoxClientApi::AccountName& AccountName, const VivoxClientApi::Uri& TargetUser, const VivoxClientApi::Uri& ChannelUri, int Volume, const VivoxClientApi::VCSStatus& Status)
 {
 	UE_LOG(LogVivoxVoiceChat, Warning, TEXT("onSetParticipantAudioOutputDeviceVolumeForMeFailed channel:%s user:%s volume:%i error:%s (%i)"), ANSI_TO_TCHAR(ChannelUri.ToString()), ANSI_TO_TCHAR(TargetUser.ToString()), Volume, ANSI_TO_TCHAR(Status.ToString()), Status.GetStatusCode());
 
 	FChannelSession& ChannelSession = GetChannelSession(ChannelUri);
-	FString PlayerName = GetPlayerNameFromUri(TargetUser);
+	FString PlayerName = VivoxVoiceChat.GetPlayerNameFromUri(TargetUser);
 	if (FParticipant* ChannelParticipant = ChannelSession.Participants.Find(PlayerName))
 	{
 		// TODO: should this retry setting volume?
 	}
 }
 
-void FVivoxVoiceChat::onSetChannelAudioOutputDeviceVolumeCompleted(const VivoxClientApi::AccountName& AccountName, const VivoxClientApi::Uri& ChannelUri, int Volume)
+void FVivoxVoiceChatUser::onSetChannelAudioOutputDeviceVolumeCompleted(const VivoxClientApi::AccountName& AccountName, const VivoxClientApi::Uri& ChannelUri, int Volume)
 {
 	UE_LOG(LogVivoxVoiceChat, Verbose, TEXT("onSetChannelAudioOutputDeviceVolumeCompleted channel:%s volume:%i"), ANSI_TO_TCHAR(ChannelUri.ToString()), Volume);
 }
 
-void FVivoxVoiceChat::onSetChannelAudioOutputDeviceVolumeFailed(const VivoxClientApi::AccountName& AccountName, const VivoxClientApi::Uri& ChannelUri, int Volume, const VivoxClientApi::VCSStatus& Status)
+void FVivoxVoiceChatUser::onSetChannelAudioOutputDeviceVolumeFailed(const VivoxClientApi::AccountName& AccountName, const VivoxClientApi::Uri& ChannelUri, int Volume, const VivoxClientApi::VCSStatus& Status)
 {
 	UE_LOG(LogVivoxVoiceChat, Warning, TEXT("onSetChannelAudioOutputDeviceVolumeFailed channel:%s volume:%i error:%s (%i)"), ANSI_TO_TCHAR(ChannelUri.ToString()), Volume, ANSI_TO_TCHAR(Status.ToString()), Status.GetStatusCode());
 }
 
-void FVivoxVoiceChat::onSetParticipantMutedForMeCompleted(const VivoxClientApi::AccountName& AccountName, const VivoxClientApi::Uri& Target, const VivoxClientApi::Uri& ChannelUri, bool bMuted)
+void FVivoxVoiceChatUser::onSetParticipantMutedForMeCompleted(const VivoxClientApi::AccountName& AccountName, const VivoxClientApi::Uri& Target, const VivoxClientApi::Uri& ChannelUri, bool bMuted)
 {
 	UE_LOG(LogVivoxVoiceChat, Verbose, TEXT("onSetParticipantMutedForMeCompleted channel:%s user:%s muted:%s"), ANSI_TO_TCHAR(ChannelUri.ToString()), ANSI_TO_TCHAR(Target.ToString()), *LexToString(bMuted));
 
 	FChannelSession& ChannelSession = GetChannelSession(ChannelUri);
-	FString PlayerName = GetPlayerNameFromUri(Target);
+	FString PlayerName = VivoxVoiceChat.GetPlayerNameFromUri(Target);
 	if (FParticipant* ChannelParticipant = ChannelSession.Participants.Find(PlayerName))
 	{
 		// TODO: Determine how should this interact with mutes from setting volume to 0
@@ -1607,60 +1656,65 @@
 	}
 }
 
-void FVivoxVoiceChat::onSetParticipantMutedForMeFailed(const VivoxClientApi::AccountName& AccountName, const VivoxClientApi::Uri& Target, const VivoxClientApi::Uri& ChannelUri, bool bMuted, const VivoxClientApi::VCSStatus& Status)
+void FVivoxVoiceChatUser::onSetParticipantMutedForMeFailed(const VivoxClientApi::AccountName& AccountName, const VivoxClientApi::Uri& Target, const VivoxClientApi::Uri& ChannelUri, bool bMuted, const VivoxClientApi::VCSStatus& Status)
 {
 	UE_LOG(LogVivoxVoiceChat, Warning, TEXT("onSetParticipantMutedForMeFailed channel:%s user:%s muted:%s error:%s (%i)"), ANSI_TO_TCHAR(ChannelUri.ToString()), ANSI_TO_TCHAR(Target.ToString()), *LexToString(bMuted), ANSI_TO_TCHAR(Status.ToString()), Status.GetStatusCode());
 
 	FChannelSession& ChannelSession = GetChannelSession(ChannelUri);
-	FString PlayerName = GetPlayerNameFromUri(Target);
+	FString PlayerName = VivoxVoiceChat.GetPlayerNameFromUri(Target);
 	if (FParticipant* ChannelParticipant = ChannelSession.Participants.Find(PlayerName))
 	{
 		// TODO: should this retry mute?
 	}
 }
 
-void FVivoxVoiceChat::onSetChannelTransmissionToSpecificChannelCompleted(const VivoxClientApi::AccountName& AccountName, const VivoxClientApi::Uri& ChannelUri)
+void FVivoxVoiceChatUser::onSetChannelTransmissionToSpecificChannelCompleted(const VivoxClientApi::AccountName& AccountName, const VivoxClientApi::Uri& ChannelUri)
 {
 	UE_LOG(LogVivoxVoiceChat, Verbose, TEXT("onSetChannelTransmissionToSpecificChannelCompleted channel:%s"), ANSI_TO_TCHAR(ChannelUri.ToString()));
 }
 
-void FVivoxVoiceChat::onSetChannelTransmissionToSpecificChannelFailed(const VivoxClientApi::AccountName& AccountName, const VivoxClientApi::Uri& ChannelUri, const VivoxClientApi::VCSStatus& Status)
+void FVivoxVoiceChatUser::onSetChannelTransmissionToSpecificChannelFailed(const VivoxClientApi::AccountName& AccountName, const VivoxClientApi::Uri& ChannelUri, const VivoxClientApi::VCSStatus& Status)
 {
 	UE_LOG(LogVivoxVoiceChat, Warning, TEXT("onSetChannelTransmissionToSpecificChannelCompleted channel:%s error:%s (%i)"), ANSI_TO_TCHAR(ChannelUri.ToString()), ANSI_TO_TCHAR(Status.ToString()), Status.GetStatusCode());
 }
 
-void FVivoxVoiceChat::onSetChannelTransmissionToAllCompleted(const VivoxClientApi::AccountName& AccountName)
+void FVivoxVoiceChatUser::onSetChannelTransmissionToAllCompleted(const VivoxClientApi::AccountName& AccountName)
 {
 	UE_LOG(LogVivoxVoiceChat, Verbose, TEXT("onSetChannelTransmissionToAllCompleted"));
 }
 
-void FVivoxVoiceChat::onSetChannelTransmissionToAllFailed(const VivoxClientApi::AccountName& AccountName, const VivoxClientApi::VCSStatus& Status)
+void FVivoxVoiceChatUser::onSetChannelTransmissionToAllFailed(const VivoxClientApi::AccountName& AccountName, const VivoxClientApi::VCSStatus& Status)
 {
 	UE_LOG(LogVivoxVoiceChat, Warning, TEXT("onSetChannelTransmissionToAllFailed error:%s (%i)"), ANSI_TO_TCHAR(Status.ToString()), Status.GetStatusCode());
 }
 
-void FVivoxVoiceChat::onSetChannelTransmissionToNoneCompleted(const VivoxClientApi::AccountName& AccountName)
+void FVivoxVoiceChatUser::onSetChannelTransmissionToNoneCompleted(const VivoxClientApi::AccountName& AccountName)
 {
 	UE_LOG(LogVivoxVoiceChat, Verbose, TEXT("onSetChannelTransmissionToNoneCompleted"));
 }
 
-void FVivoxVoiceChat::onSetChannelTransmissionToNoneFailed(const VivoxClientApi::AccountName& AccountName, const VivoxClientApi::VCSStatus& Status)
+void FVivoxVoiceChatUser::onSetChannelTransmissionToNoneFailed(const VivoxClientApi::AccountName& AccountName, const VivoxClientApi::VCSStatus& Status)
 {
 	UE_LOG(LogVivoxVoiceChat, Warning, TEXT("onSetChannelTransmissionToNoneFailed error:%s (%i)"), ANSI_TO_TCHAR(Status.ToString()), Status.GetStatusCode());
 }
 
-void FVivoxVoiceChat::onAudioUnitStarted(const VivoxClientApi::Uri& InitialTargetUri)
+void FVivoxVoiceChatUser::onAudioUnitStarted(const char* SessionGroupHandle, const VivoxClientApi::Uri& InitialTargetUri)
 {
 
 }
 
-void FVivoxVoiceChat::onAudioUnitStopped(const VivoxClientApi::Uri& InitialTargetUri)
+void FVivoxVoiceChatUser::onAudioUnitStopped(const char* SessionGroupHandle, const VivoxClientApi::Uri& InitialTargetUri)
 {
 
 }
 
-void FVivoxVoiceChat::onAudioUnitAfterCaptureAudioRead(const VivoxClientApi::Uri& InitialTargetUri, short* PcmFrames, int PcmFrameCount, int AudioFrameRate, int ChannelsPerFrame)
+void FVivoxVoiceChatUser::onAudioUnitAfterCaptureAudioRead(const char* SessionGroupHandle, const VivoxClientApi::Uri& InitialTargetUri, short* PcmFrames, int PcmFrameCount, int AudioFrameRate, int ChannelsPerFrame)
 {
+	if (SessionGroup != ANSI_TO_TCHAR(SessionGroupHandle))
+	{
+		return;
+	}
+
 	if (InitialTargetUri.IsValid())
 	{
 		FScopeLock Lock(&AfterCaptureAudioReadLock);
@@ -1673,8 +1727,13 @@
 	}
 }
 
-void FVivoxVoiceChat::onAudioUnitBeforeCaptureAudioSent(const VivoxClientApi::Uri& InitialTargetUri, short* PcmFrames, int PcmFrameCount, int AudioFrameRate, int ChannelsPerFrame, bool bSpeaking)
+void FVivoxVoiceChatUser::onAudioUnitBeforeCaptureAudioSent(const char* SessionGroupHandle, const VivoxClientApi::Uri& InitialTargetUri, short* PcmFrames, int PcmFrameCount, int AudioFrameRate, int ChannelsPerFrame, bool bSpeaking)
 {
+	if (SessionGroup != ANSI_TO_TCHAR(SessionGroupHandle))
+	{
+		return;
+	}
+
 	if (InitialTargetUri.IsValid())
 	{
 		FScopeLock Lock(&BeforeCaptureAudioSentLock);
@@ -1682,8 +1741,13 @@
 	}
 }
 
-void FVivoxVoiceChat::onAudioUnitBeforeRecvAudioRendered(const VivoxClientApi::Uri& InitialTargetUri, short* PcmFrames, int PcmFrameCount, int AudioFrameRate, int ChannelsPerFrame, bool bSilence)
+void FVivoxVoiceChatUser::onAudioUnitBeforeRecvAudioRendered(const char* SessionGroupHandle, const VivoxClientApi::Uri& InitialTargetUri, short* PcmFrames, int PcmFrameCount, int AudioFrameRate, int ChannelsPerFrame, bool bSilence)
 {
+	if (SessionGroup != ANSI_TO_TCHAR(SessionGroupHandle))
+	{
+		return;
+	}
+
 	if (InitialTargetUri.IsValid())
 	{
 		FScopeLock Lock(&BeforeRecvAudioRenderedLock);
@@ -1691,6 +1755,185 @@
 	}
 }
 
+static void* VivoxMalloc(size_t bytes)
+{
+	LLM_SCOPE(ELLMTag::AudioVoiceChat);
+	return FMemory::Malloc(bytes);
+}
+
+static void VivoxFree(void* ptr)
+{
+	LLM_SCOPE(ELLMTag::AudioVoiceChat);
+	FMemory::Free(ptr);
+}
+
+static void* VivoxRealloc(void* ptr, size_t bytes)
+{
+	LLM_SCOPE(ELLMTag::AudioVoiceChat);
+	return FMemory::Realloc(ptr, bytes);
+}
+
+static void* VivoxCalloc(size_t num, size_t bytes)
+{
+	LLM_SCOPE(ELLMTag::AudioVoiceChat);
+	const size_t Size = bytes * num;
+	void* Ret = FMemory::Malloc(Size);
+	FMemory::Memzero(Ret, Size);
+	return Ret;
+}
+
+static void* VivoxMallocAligned(size_t alignment, size_t bytes)
+{
+	LLM_SCOPE(ELLMTag::AudioVoiceChat);
+	return FMemory::Malloc(bytes, alignment);
+}
+
+static void VivoxFreeAligned(void* ptr)
+{
+	LLM_SCOPE(ELLMTag::AudioVoiceChat);
+	FMemory::Free(ptr);
+}
+
+bool FVivoxVoiceChat::Initialize()
+{
+	if (!IsInitialized())
+	{
+		bool bEnabled = true;
+		GConfig->GetBool(TEXT("VoiceChat.Vivox"), TEXT("bEnabled"), bEnabled, GEngineIni);
+		if (bEnabled)
+		{
+			GConfig->GetString(TEXT("VoiceChat.Vivox"), TEXT("ServerUrl"), VivoxServerUrl, GEngineIni);
+			GConfig->GetString(TEXT("VoiceChat.Vivox"), TEXT("Domain"), VivoxDomain, GEngineIni);
+			GConfig->GetString(TEXT("VoiceChat.Vivox"), TEXT("Issuer"), VivoxIssuer, GEngineIni);
+			GConfig->GetString(TEXT("VoiceChat.Vivox"), TEXT("Namespace"), VivoxNamespace, GEngineIni);
+			GConfig->GetString(TEXT("VoiceChat.Vivox"), TEXT("InsecureSecret"), VivoxInsecureSecret, GEngineIni);
+
+			if (VivoxNamespace.IsEmpty())
+			{
+				VivoxNamespace = VivoxIssuer;
+			}
+
+			// positional audio settings
+			AttenuationModel = EVoiceChatAttenuationModel::InverseByDistance;
+			MinDistance = 100;
+			MaxDistance = 3000;
+			Rolloff = 1.0f;
+			FString AttenuationModelString;
+			if (GConfig->GetString(TEXT("VoiceChat.Vivox"), TEXT("AttenuationModel"), AttenuationModelString, GEngineIni) && !AttenuationModelString.IsEmpty())
+			{
+				if (AttenuationModelString == TEXT("None"))
+				{
+					AttenuationModel = EVoiceChatAttenuationModel::None;
+				}
+				else if (AttenuationModelString == TEXT("InverseByDistance"))
+				{
+					AttenuationModel = EVoiceChatAttenuationModel::InverseByDistance;
+				}
+				else if (AttenuationModelString == TEXT("LinearByDistance"))
+				{
+					AttenuationModel = EVoiceChatAttenuationModel::LinearByDistance;
+				}
+				else if (AttenuationModelString == TEXT("ExponentialByDistance"))
+				{
+					AttenuationModel = EVoiceChatAttenuationModel::ExponentialByDistance;
+				}
+				else
+				{
+					UE_LOG(LogVivoxVoiceChat, Warning, TEXT("Unknown AttenuationModel: %s"), *AttenuationModelString);
+				}
+			}
+			GConfig->GetInt(TEXT("VoiceChat.Vivox"), TEXT("MinDistance"), MinDistance, GEngineIni);
+			GConfig->GetInt(TEXT("VoiceChat.Vivox"), TEXT("MaxDistance"), MaxDistance, GEngineIni);
+			GConfig->GetFloat(TEXT("VoiceChat.Vivox"), TEXT("Rolloff"), Rolloff, GEngineIni);
+
+			const char* VivoxVersionInfo = vx_get_sdk_version_info();
+			UE_LOG(LogVivoxVoiceChat, Log, TEXT("Initializing Vivox %s"), ANSI_TO_TCHAR(VivoxVersionInfo));
+
+			vx_sdk_config_t ConfigHints;
+			int Result = vx_get_default_config3(&ConfigHints, sizeof(ConfigHints));
+			if (Result != 0)
+			{
+				UE_LOG(LogVivoxVoiceChat, Warning, TEXT("Failed to get default config: error:%s (%i)"), ANSI_TO_TCHAR(VivoxClientApi::GetErrorString(Result)), Result);
+			}
+			else
+			{
+				SetVivoxSdkConfigHints(ConfigHints);
+
+				LogLevel VivoxLogLevel = LogLevelWarning;
+				FString LogLevelString;
+				if (GConfig->GetString(TEXT("VoiceChat.Vivox"), TEXT("LogLevel"), LogLevelString, GEngineIni))
+				{
+					if (LogLevelString == TEXT("None"))
+					{
+						VivoxLogLevel = LogLevelNone;
+					}
+					else if (LogLevelString == TEXT("Error"))
+					{
+						VivoxLogLevel = LogLevelError;
+					}
+					else if (LogLevelString == TEXT("Warning"))
+					{
+						VivoxLogLevel = LogLevelWarning;
+					}
+					else if (LogLevelString == TEXT("Info"))
+					{
+						VivoxLogLevel = LogLevelInfo;
+					}
+					else if (LogLevelString == TEXT("Debug"))
+					{
+						VivoxLogLevel = LogLevelDebug;
+					}
+					else if (LogLevelString == TEXT("Trace"))
+					{
+						VivoxLogLevel = LogLevelTrace;
+					}
+				}
+
+				VivoxClientApi::VCSStatus Status = VivoxClientConnection.Initialize(this, VivoxLogLevel, true, true, &ConfigHints, sizeof(ConfigHints));
+				if (Status.IsError())
+				{
+					const FVoiceChatResult VoiceChatResult = ResultFromVivoxStatus(Status);
+					UE_LOG(LogVivoxVoiceChat, Warning, TEXT("Initialize %s"), *LexToString(VoiceChatResult));
+				}
+				else
+				{
+					bInitialized = true;
+
+					bool bVADAutomaticParameterSelection = true;
+					GConfig->GetBool(TEXT("VoiceChat.Vivox"), TEXT("bVADAutomaticParameterSelection"), bVADAutomaticParameterSelection, GEngineIni);
+					VivoxClientConnection.SetVADAutomaticParameterSelection(VivoxClientApi::AccountName(), bVADAutomaticParameterSelection);
+				}
+			}
+		}
+	}
+
+	return IsInitialized();
+}
+
+bool FVivoxVoiceChat::Uninitialize()
+{
+	if (IsInitialized())
+	{
+		VivoxClientConnection.Uninitialize();
+		ConnectionState = EConnectionState::Disconnected;
+		DispatchAll(&FVivoxVoiceChatUser::onDisconnected, VivoxClientApi::Uri(TCHAR_TO_ANSI(*VivoxServerUrl)), VivoxClientApi::VCSStatus(VX_E_NOT_INITIALIZED));
+		bInitialized = false;
+	}
+
+	return true;
+}
+
+bool FVivoxVoiceChat::IsInitialized() const
+{
+	return bInitialized;
+}
+
+IVoiceChatUser* FVivoxVoiceChat::CreateUser()
+{
+	return new FVivoxVoiceChatUser(*this);
+}
+
+
 void FVivoxVoiceChat::SetVivoxSdkConfigHints(vx_sdk_config_t& Hints)
 {
 	Hints.pf_malloc_func = &VivoxMalloc;
@@ -1827,7 +2070,7 @@
 		int32 Channel3dParametersIndex = Channel.Find(TEXT("!p-"));
 		if (Channel3dParametersIndex != INDEX_NONE)
 		{
-			Channel.LeftInline(Channel3dParametersIndex);
+			Channel = Channel.LeftInline(Channel3dParametersIndex);
 		}
 		return Channel;
 	}
@@ -1852,7 +2095,7 @@
 	}
 }
 
-FVivoxVoiceChat::FParticipant& FVivoxVoiceChat::GetParticipant(const FString& PlayerName)
+FVivoxVoiceChatUser::FParticipant& FVivoxVoiceChatUser::GetParticipant(const FString& PlayerName)
 {
 	if (FParticipant* Participant = LoginSession.Participants.Find(PlayerName))
 	{
@@ -1861,11 +2104,11 @@
 
 	FParticipant& NewParticipant = LoginSession.Participants.Add(PlayerName);
 	NewParticipant.PlayerName = PlayerName;
-	NewParticipant.UserUri = CreateUserUri(PlayerName);
+	NewParticipant.UserUri = VivoxVoiceChat.CreateUserUri(PlayerName);
 	return NewParticipant;
 }
 
-const FVivoxVoiceChat::FParticipant& FVivoxVoiceChat::GetParticipant(const FString& PlayerName) const
+const FVivoxVoiceChatUser::FParticipant& FVivoxVoiceChatUser::GetParticipant(const FString& PlayerName) const
 {
 	if (const FParticipant* Participant = LoginSession.Participants.Find(PlayerName))
 	{
@@ -1878,7 +2121,7 @@
 	}
 }
 
-FVivoxVoiceChat::FChannelSession& FVivoxVoiceChat::GetChannelSession(const FString& ChannelName)
+FVivoxVoiceChatUser::FChannelSession& FVivoxVoiceChatUser::GetChannelSession(const FString& ChannelName)
 {
 	if (FChannelSession* Session = LoginSession.ChannelSessions.Find(ChannelName))
 	{
@@ -1888,11 +2131,11 @@
 	FChannelSession& NewSession = LoginSession.ChannelSessions.Add(ChannelName);
 	NewSession.ChannelName = ChannelName;
 	NewSession.ChannelType = EVoiceChatChannelType::NonPositional;
-	NewSession.ChannelUri = CreateChannelUri(ChannelName, NewSession.ChannelType);
+	NewSession.ChannelUri = VivoxVoiceChat.CreateChannelUri(ChannelName, NewSession.ChannelType);
 	return NewSession;
 }
 
-const FVivoxVoiceChat::FChannelSession& FVivoxVoiceChat::GetChannelSession(const FString& ChannelName) const
+const FVivoxVoiceChatUser::FChannelSession& FVivoxVoiceChatUser::GetChannelSession(const FString& ChannelName) const
 {
 	if (const FChannelSession* Session = LoginSession.ChannelSessions.Find(ChannelName))
 	{
@@ -1905,146 +2148,91 @@
 	}
 }
 
-FVivoxVoiceChat::FChannelSession& FVivoxVoiceChat::GetChannelSession(const VivoxClientApi::Uri& ChannelUri)
+FVivoxVoiceChatUser::FChannelSession& FVivoxVoiceChatUser::GetChannelSession(const VivoxClientApi::Uri& ChannelUri)
 {
-	FString ChannelName = GetChannelNameFromUri(ChannelUri);
+	FString ChannelName = VivoxVoiceChat.GetChannelNameFromUri(ChannelUri);
 	if (FChannelSession* Session = LoginSession.ChannelSessions.Find(ChannelName))
 	{
 		return *Session;
 	}
 	FChannelSession& Session = LoginSession.ChannelSessions.Add(ChannelName);
 	Session.ChannelName = ChannelName;
-	Session.ChannelType = GetChannelTypeFromUri(ChannelUri);
+	Session.ChannelType = VivoxVoiceChat.GetChannelTypeFromUri(ChannelUri);
 	Session.ChannelUri = ChannelUri;
 	return Session;
 }
 
-void FVivoxVoiceChat::RemoveChannelSession(const FString& ChannelName)
+void FVivoxVoiceChatUser::RemoveChannelSession(const FString& ChannelName)
 {
 	// TODO: Should this trigger participant leave delegates?
 	LoginSession.ChannelSessions.Remove(ChannelName);
 }
 
-void FVivoxVoiceChat::ClearChannelSessions()
+void FVivoxVoiceChatUser::ClearChannelSessions()
 {
 	// TODO: Should this trigger channel/participant leave delegates?
 	LoginSession.ChannelSessions.Reset();
 }
 
-void FVivoxVoiceChat::ClearLoginSession()
+void FVivoxVoiceChatUser::ClearLoginSession()
 {
 	ClearChannelSessions();
 	LoginSession.State = FLoginSession::EState::LoggedOut;
 }
 
-void FVivoxVoiceChat::ApplyAudioInputOptions()
+void FVivoxVoiceChatUser::ApplyAudioInputOptions()
 {
-	VivoxClientConnection.SetMasterAudioInputDeviceVolume(FMath::Lerp(VIVOX_MIN_VOL, VIVOX_MAX_VOL, AudioInputOptions.Volume));
+	VivoxClientConnection.SetMasterAudioInputDeviceVolume(LoginSession.AccountName, FMath::Lerp(VIVOX_MIN_VOL, VIVOX_MAX_VOL, AudioInputOptions.Volume));
 
 	const bool bVolumeMuted = AudioInputOptions.Volume < SMALL_NUMBER;
-	VivoxClientConnection.SetAudioInputDeviceMuted(AudioInputOptions.bMuted || bVolumeMuted);
+	VivoxClientConnection.SetAudioInputDeviceMuted(LoginSession.AccountName, AudioInputOptions.bMuted || bVolumeMuted);
 }
 
-void FVivoxVoiceChat::ApplyAudioOutputOptions()
+void FVivoxVoiceChatUser::ApplyAudioOutputOptions()
 {
-	VivoxClientConnection.SetMasterAudioOutputDeviceVolume(FMath::Lerp(VIVOX_MIN_VOL, VIVOX_MAX_VOL, AudioOutputOptions.Volume));
+	VivoxClientConnection.SetMasterAudioOutputDeviceVolume(LoginSession.AccountName, FMath::Lerp(VIVOX_MIN_VOL, VIVOX_MAX_VOL, AudioOutputOptions.Volume));
 
 	const bool bVolumeMuted = AudioOutputOptions.Volume < SMALL_NUMBER;
-	VivoxClientConnection.SetAudioOutputDeviceMuted(AudioOutputOptions.bMuted || bVolumeMuted);
+	VivoxClientConnection.SetAudioOutputDeviceMuted(LoginSession.AccountName, AudioOutputOptions.bMuted || bVolumeMuted);
 }
 
 bool FVivoxVoiceChat::Exec(UWorld* InWorld, const TCHAR* Cmd, FOutputDevice& Ar)
 {
+#define VIVOX_EXEC_LOG(Fmt, ...) Ar.CategorizedLogf(LogVivoxVoiceChat.GetCategoryName(), ELogVerbosity::Log, Fmt, ##__VA_ARGS__)
+
 	if (FParse::Command(&Cmd, TEXT("VIVOX")))
 	{
-		if (FParse::Command(&Cmd, TEXT("INFO")))
+		const TCHAR* SubCmd = Cmd;
+		if (FParse::Command(&Cmd, TEXT("LIST")))
 		{
-			UE_LOG(LogVivoxVoiceChat, Log, TEXT("Initialized: %s"), *LexToString(IsInitialized()));
+			for (int UserIndex = 0; UserIndex < VoiceChatUsers.Num(); ++UserIndex)
+			{
+				if (FVivoxVoiceChatUser* User = VoiceChatUsers[UserIndex])
+				{
+					VIVOX_EXEC_LOG(TEXT("VivoxUser Index:%i PlayerName:%s"), UserIndex, *User->GetLoggedInPlayerName());
+				}
+			}
+			return true;
+		}
+		else if (FParse::Command(&Cmd, TEXT("INFO")))
+		{
+			VIVOX_EXEC_LOG(TEXT("Initialized: %s"), *LexToString(IsInitialized()));
 			if (IsInitialized())
 			{
-				UE_LOG(LogVivoxVoiceChat, Log, TEXT("  Input Devices: muted:%s volume:%.2f"), *LexToString(GetAudioInputDeviceMuted()), GetAudioInputVolume());
-				const FString InputDevice = GetInputDevice();
-				const FString DefaultInputDevice = GetDefaultInputDevice();
-				if (InputDevice == DefaultInputDevice)
-				{
-					UE_LOG(LogVivoxVoiceChat, Log, TEXT("    [%s] (Selected) (Default)"), *DefaultInputDevice);
-				}
-				else
-				{
-					UE_LOG(LogVivoxVoiceChat, Log, TEXT("    [%s] (Selected)"), *InputDevice);
-					UE_LOG(LogVivoxVoiceChat, Log, TEXT("    [%s] (Default)"), *DefaultInputDevice);
-				}
-				for (const FString& Device : GetAvailableInputDevices())
-				{
-					if (Device != DefaultInputDevice && Device != InputDevice)
-					{
-						UE_LOG(LogVivoxVoiceChat, Log, TEXT("    [%s]"), *Device);
-					}
-				}
-
-				UE_LOG(LogVivoxVoiceChat, Log, TEXT("  Output Devices: muted:%s volume:%.2f"), *LexToString(GetAudioOutputDeviceMuted()), GetAudioOutputVolume());
-				const FString OutputDevice = GetOutputDevice();
-				const FString DefaultOutputDevice = GetDefaultOutputDevice();
-				if (OutputDevice == DefaultOutputDevice)
-				{
-					UE_LOG(LogVivoxVoiceChat, Log, TEXT("    [%s] (Selected) (Default)"), *DefaultOutputDevice);
-				}
-				else
-				{
-					UE_LOG(LogVivoxVoiceChat, Log, TEXT("    [%s] (Selected)"), *OutputDevice);
-					UE_LOG(LogVivoxVoiceChat, Log, TEXT("    [%s] (Default)"), *DefaultOutputDevice);
-				}
-				for (const FString& Device : GetAvailableOutputDevices())
-				{
-					if (Device != DefaultOutputDevice && Device != OutputDevice)
-					{
-						UE_LOG(LogVivoxVoiceChat, Log, TEXT("    [%s]"), *Device);
-					}
-				}
-
-				UE_LOG(LogVivoxVoiceChat, Log, TEXT("Connection Status: %s"), *ToString(ConnectionState));
+				VIVOX_EXEC_LOG(TEXT("Connection Status: %s"), *ToString(ConnectionState));
 				if (IsConnected())
 				{
-					UE_LOG(LogVivoxVoiceChat, Log, TEXT("  Server: %s"), *VivoxServerUrl);
-					UE_LOG(LogVivoxVoiceChat, Log, TEXT("  Domain: %s"), *VivoxDomain);
-					UE_LOG(LogVivoxVoiceChat, Log, TEXT("Login Status: %s"), *ToString(LoginSession.State));
-					if (IsLoggedIn())
-					{
-						UE_LOG(LogVivoxVoiceChat, Log, TEXT("  PlayerName: %s"), *LoginSession.PlayerName);
-						UE_LOG(LogVivoxVoiceChat, Log, TEXT("  AccountName: %s"), ANSI_TO_TCHAR(LoginSession.AccountName.ToString()));
-						UE_LOG(LogVivoxVoiceChat, Log, TEXT("  UserUri: %s"), ANSI_TO_TCHAR(LoginSession.UserUri.ToString()));
+					VIVOX_EXEC_LOG(TEXT("  Server: %s"), *VivoxServerUrl);
+					VIVOX_EXEC_LOG(TEXT("  Domain: %s"), *VivoxDomain);
+				}
+			}
 
-						VivoxClientApi::ChannelTransmissionPolicy TransmissionPolicy = VivoxClientConnection.GetChannelTransmissionPolicy(LoginSession.AccountName);
-						FString TransmitString;
-						switch (GetTransmitMode())
-						{
-						case EVoiceChatTransmitMode::All:
-							TransmitString = TEXT("ALL");
-							break;
-						case EVoiceChatTransmitMode::None:
-							TransmitString = TEXT("NONE");
-							break;
-						case EVoiceChatTransmitMode::Channel:
-							TransmitString = FString::Printf(TEXT("CHANNEL:%s"), *GetTransmitChannel());
-							break;
-						}
-						UE_LOG(LogVivoxVoiceChat, Log, TEXT("Channels: transmitting:%s"), *TransmitString);
-						for (const TPair<FString, FChannelSession>& ChannelSessionPair : LoginSession.ChannelSessions)
-						{
-							const FString& ChannelName = ChannelSessionPair.Key;
-							const FChannelSession& ChannelSession = ChannelSessionPair.Value;
-							UE_LOG(LogVivoxVoiceChat, Log, TEXT("  %s"), *ChannelName);
-							UE_LOG(LogVivoxVoiceChat, Log, TEXT("    Channel Status: %s"), *ToString(ChannelSession.State));
-							UE_LOG(LogVivoxVoiceChat, Log, TEXT("    Channel Uri: %s"), ANSI_TO_TCHAR(ChannelSession.ChannelUri.ToString()));
-							UE_LOG(LogVivoxVoiceChat, Log, TEXT("    Participants:"));
-							for (const TPair<FString, FParticipant>& ParticipantPair : ChannelSession.Participants)
-							{
-								const FString& ParticipantName = ParticipantPair.Key;
-								const FParticipant& Participant = ParticipantPair.Value;
-								UE_LOG(LogVivoxVoiceChat, Log, TEXT("      %s uri:%s talking:%s muted:%s volume:%.2f"), *ParticipantName, ANSI_TO_TCHAR(Participant.UserUri.ToString()), *LexToString(Participant.bTalking), *LexToString(Participant.bMuted), Participant.Volume);
-							}
-						}
-					}
+			for (int UserIndex = 0; UserIndex < VoiceChatUsers.Num(); ++UserIndex)
+			{
+				if (FVivoxVoiceChatUser* User = VoiceChatUsers[UserIndex])
+				{
+					VIVOX_EXEC_LOG(TEXT("VivoxUser Index:%i PlayerName:%s"), UserIndex, *User->GetLoggedInPlayerName());
+					User->Exec(InWorld, SubCmd, Ar);
 				}
 			}
 			return true;
@@ -2060,264 +2248,30 @@
 			Uninitialize();
 			return true;
 		}
-		else if (FParse::Command(&Cmd, TEXT("INPUT")))
+		else
 		{
-			if (FParse::Command(&Cmd, TEXT("SETVOLUME")))
+			int Index = 0;
+			if (FParse::Value(Cmd, TEXT("INDEX="), Index))
 			{
-				FString Volume;
-				if (FParse::Token(Cmd, Volume, false))
+				FParse::Token(Cmd, false); // skip over INDEX=#
+				if (Index < VoiceChatUsers.Num())
 				{
-					SetAudioInputVolume(FCString::Atof(*Volume));
-					return true;
-				}
-			}
-			else if (FParse::Command(&Cmd, TEXT("MUTE")))
-			{
-				SetAudioInputDeviceMuted(true);
-				return true;
-			}
-			else if (FParse::Command(&Cmd, TEXT("UNMUTE")))
-			{
-				SetAudioInputDeviceMuted(false);
-				return true;
-			}
-			else if (FParse::Command(&Cmd, TEXT("LISTDEVICES")))
-			{
-				TArray<FString> Devices = GetAvailableInputDevices();
-				UE_LOG(LogVivoxVoiceChat, Log, TEXT("Input Devices:"));
-				for (const FString& Device : Devices)
-				{
-					UE_LOG(LogVivoxVoiceChat, Log, TEXT("  %s"), *Device);
-				}
-				return true;
-			}
-			else if (FParse::Command(&Cmd, TEXT("SETDEVICE")))
-			{
-				FString Device;
-				if (FParse::Token(Cmd, Device, false))
-				{
-					SetInputDevice(Device);
-					return true;
-				}
-			}
-			else if (FParse::Command(&Cmd, TEXT("SETDEFAULTDEVICE")))
-			{
-				SetInputDevice(FString());
-				return true;
-			}
-		}
-		else if (FParse::Command(&Cmd, TEXT("OUTPUT")))
-		{
-			if (FParse::Command(&Cmd, TEXT("SETVOLUME")))
-			{
-				FString Volume;
-				if (FParse::Token(Cmd, Volume, false))
-				{
-					SetAudioOutputVolume(FCString::Atof(*Volume));
-					return true;
-				}
-			}
-			else if (FParse::Command(&Cmd, TEXT("MUTE")))
-			{
-				SetAudioOutputDeviceMuted(true);
-				return true;
-			}
-			else if (FParse::Command(&Cmd, TEXT("UNMUTE")))
-			{
-				SetAudioOutputDeviceMuted(false);
-				return true;
-			}
-			else if (FParse::Command(&Cmd, TEXT("LISTDEVICES")))
-			{
-				TArray<FString> Devices = GetAvailableOutputDevices();
-				UE_LOG(LogVivoxVoiceChat, Log, TEXT("Output Devices:"));
-				for (const FString& Device : Devices)
-				{
-					UE_LOG(LogVivoxVoiceChat, Log, TEXT("  %s"), *Device);
-				}
-				return true;
-			}
-			else if (FParse::Command(&Cmd, TEXT("SETDEVICE")))
-			{
-				FString Device;
-				if (FParse::Token(Cmd, Device, false))
-				{
-					SetOutputDevice(Device);
-					return true;
-				}
-			}
-			else if (FParse::Command(&Cmd, TEXT("SETDEFAULTDEVICE")))
-			{
-				SetOutputDevice(FString());
-				return true;
-			}
-		}
-		else if (FParse::Command(&Cmd, TEXT("CONNECT")))
-		{
-			Connect(FOnVoiceChatConnectCompleteDelegate::CreateLambda([](const FVoiceChatResult& Result)
-			{
-				UE_LOG(LogVivoxVoiceChat, Log, TEXT("VIVOX CONNECT result:%s"), *LexToString(Result));
-			}));
-			return true;
-		}
-		else if (FParse::Command(&Cmd, TEXT("DISCONNECT")))
-		{
-			Disconnect(FOnVoiceChatDisconnectCompleteDelegate::CreateLambda([](const FVoiceChatResult& Result)
-			{
-				UE_LOG(LogVivoxVoiceChat, Log, TEXT("VIVOX DISCONNECT result:%s"), *LexToString(Result));
-			}));
-			return true;
-		}
-		else if (FParse::Command(&Cmd, TEXT("LOGIN")))
-		{
-			FString PlayerName;
-			if (FParse::Token(Cmd, PlayerName, false))
-			{
-				FString Token = InsecureGetLoginToken(PlayerName);
-
-				Login(0, PlayerName, Token, FOnVoiceChatLoginCompleteDelegate::CreateLambda([](const FString& LoggedInPlayerName, const FVoiceChatResult& Result)
-				{
-					UE_LOG(LogVivoxVoiceChat, Log, TEXT("VIVOX LOGIN playername:%s result:%s"), *LoggedInPlayerName, *LexToString(Result));
-				}));
-				return true;
-			}
-		}
-		else if (FParse::Command(&Cmd, TEXT("LOGOUT")))
-		{
-			Logout(FOnVoiceChatLogoutCompleteDelegate::CreateLambda([](const FString& PlayerName, const FVoiceChatResult& Result)
-			{
-				UE_LOG(LogVivoxVoiceChat, Log, TEXT("VIVOX LOGOUT playername:%s result:%s"), *PlayerName, *LexToString(Result));
-			}));
-			return true;
-		}
-		else if (FParse::Command(&Cmd, TEXT("CHANNEL")))
-		{
-			if (FParse::Command(&Cmd, TEXT("JOIN")))
-			{
-				FString ChannelName;
-				if (FParse::Token(Cmd, ChannelName, false))
-				{
-					FString ChannelTypeString;
-					EVoiceChatChannelType ChannelType = EVoiceChatChannelType::NonPositional;
-					TOptional<FVoiceChatChannel3dProperties> Channel3dProperties;
-					if (FParse::Token(Cmd, ChannelTypeString, false))
+					if (FVivoxVoiceChatUser* User = VoiceChatUsers[Index])
 					{
-						if (ChannelTypeString == TEXT("POSITIONAL"))
-						{
-							ChannelType = EVoiceChatChannelType::Positional;
-						}
-						else if (ChannelTypeString == TEXT("ECHO"))
-						{
-							ChannelType = EVoiceChatChannelType::Echo;
-						}
-					}
-
-					FString Token = InsecureGetJoinToken(ChannelName, ChannelType, Channel3dProperties);
-
-					JoinChannel(ChannelName, Token, ChannelType, FOnVoiceChatChannelJoinCompleteDelegate::CreateLambda([](const FString& JoinedChannelName, const FVoiceChatResult& Result)
-					{
-						UE_LOG(LogVivoxVoiceChat, Log, TEXT("VIVOX CHANNEL JOIN channelname:%s result:%s"), *JoinedChannelName, *LexToString(Result));
-					}), Channel3dProperties);
-					return true;
-				}
-			}
-			else if (FParse::Command(&Cmd, TEXT("LEAVE")))
-			{
-				FString ChannelName;
-				if (FParse::Token(Cmd, ChannelName, false))
-				{
-					LeaveChannel(ChannelName, FOnVoiceChatChannelLeaveCompleteDelegate::CreateLambda([](const FString& LeftChannelName, const FVoiceChatResult& Result)
-					{
-						UE_LOG(LogVivoxVoiceChat, Log, TEXT("VIVOX CHANNEL LEAVE channelname:%s result:%s"), *LeftChannelName, *LexToString(Result));
-					}));
-					return true;
-				}
-			}
-			else if (FParse::Command(&Cmd, TEXT("TRANSMIT")))
-			{
-				FString ChannelName;
-				if (FParse::Token(Cmd, ChannelName, false))
-				{
-					TransmitToSpecificChannel(ChannelName);
-					return true;
-				}
-			}
-			else if (FParse::Command(&Cmd, TEXT("TRANSMITALL")))
-			{
-				TransmitToAllChannels();
-				return true;
-			}
-			else if (FParse::Command(&Cmd, TEXT("TRANSMITNONE")))
-			{
-				TransmitToNoChannels();
-				return true;
-			}
-		}
-		else if (FParse::Command(&Cmd, TEXT("PLAYER")))
-		{
-			if (FParse::Command(&Cmd, TEXT("MUTE")))
-			{
-				FString PlayerName;
-				if (FParse::Token(Cmd, PlayerName, false))
-				{
-					SetPlayerMuted(PlayerName, true);
-					return true;
-				}
-			}
-			else if (FParse::Command(&Cmd, TEXT("UNMUTE")))
-			{
-				FString PlayerName;
-				if (FParse::Token(Cmd, PlayerName, false))
-				{
-					SetPlayerMuted(PlayerName, false);
-					return true;
-				}
-			}
-			else if (FParse::Command(&Cmd, TEXT("SETVOLUME")))
-			{
-				FString PlayerName;
-				if (FParse::Token(Cmd, PlayerName, false))
-				{
-					FString Volume;
-					if (FParse::Token(Cmd, Volume, false))
-					{
-						SetPlayerVolume(PlayerName, FCString::Atof(*Volume));
-						return true;
+						return User->Exec(InWorld, Cmd, Ar);
 					}
 				}
 			}
-			else if (FParse::Command(&Cmd, TEXT("BLOCK")))
+			else if (SingleUserVoiceChatUser.IsValid())
 			{
-				TArray<FString> PlayerNames;
-				FString PlayerName;
-				while (FParse::Token(Cmd, PlayerName, false))
-				{
-					PlayerNames.Add(PlayerName);
-				}
-				if (PlayerNames.Num() > 0)
-				{
-					BlockPlayers(PlayerNames);
-					return true;
-				}
-			}
-			else if (FParse::Command(&Cmd, TEXT("UNBLOCK")))
-			{
-				TArray<FString> PlayerNames;
-				FString PlayerName;
-				while (FParse::Token(Cmd, PlayerName, false))
-				{
-					PlayerNames.Add(PlayerName);
-				}
-				if (PlayerNames.Num() > 0)
-				{
-					UnblockPlayers(PlayerNames);
-					return true;
-				}
+				return SingleUserVoiceChatUser->Exec(InWorld, SubCmd, Ar);
 			}
 		}
 #endif // !UE_BUILD_SHIPPING
 	}
 
+#undef VIVOX_EXEC_LOG
+
 	return false;
 }
 
@@ -2333,7 +2287,7 @@
 	}
 }
 
-FString FVivoxVoiceChat::ToString(FLoginSession::EState State)
+FString FVivoxVoiceChatUser::ToString(FLoginSession::EState State)
 {
 	switch (State)
 	{
@@ -2345,7 +2299,7 @@
 	}
 }
 
-FString FVivoxVoiceChat::ToString(FChannelSession::EState State)
+FString FVivoxVoiceChatUser::ToString(FChannelSession::EState State)
 {
 	switch (State)
 	{
@@ -2356,3 +2310,764 @@
 	default:									return TEXT("Unknown");
 	}
 }
+
+FVivoxVoiceChat::FVivoxVoiceChat()
+	: bInitialized(false)
+	, ConnectionState(EConnectionState::Disconnected)
+{
+	SetClientConnection(&VivoxClientConnection);
+	SetAbortEnabled(false);
+}
+
+FVivoxVoiceChat::~FVivoxVoiceChat()
+{
+
+}
+
+FVivoxVoiceChatUser& FVivoxVoiceChat::GetVoiceChatUser()
+{
+	if (!SingleUserVoiceChatUser)
+	{
+		SingleUserVoiceChatUser = MakeUnique<FVivoxVoiceChatUser>(*this);
+		ensureMsgf(VoiceChatUsers.Num() == 1, TEXT("When using multiple users, all connections should be managed by an IVoiceChatUser"));
+	}
+
+	return *SingleUserVoiceChatUser;
+}
+
+FVivoxVoiceChatUser& FVivoxVoiceChat::GetVoiceChatUser() const
+{
+	return const_cast<FVivoxVoiceChat*>(this)->GetVoiceChatUser();
+}
+
+void FVivoxVoiceChat::SetSetting(const FString& Name, const FString& Value)
+{
+	GetVoiceChatUser().SetSetting(Name, Value);
+}
+
+FString FVivoxVoiceChat::GetSetting(const FString& Name)
+{
+	return GetVoiceChatUser().GetSetting(Name);
+}
+
+void FVivoxVoiceChat::SetAudioInputVolume(float Volume)
+{
+	GetVoiceChatUser().SetAudioInputVolume(Volume);
+}
+
+void FVivoxVoiceChat::SetAudioOutputVolume(float Volume)
+{
+	GetVoiceChatUser().SetAudioOutputVolume(Volume);
+}
+
+float FVivoxVoiceChat::GetAudioInputVolume() const
+{
+	return GetVoiceChatUser().GetAudioInputVolume();
+}
+
+float FVivoxVoiceChat::GetAudioOutputVolume() const
+{
+	return GetVoiceChatUser().GetAudioOutputVolume();
+}
+
+void FVivoxVoiceChat::SetAudioInputDeviceMuted(bool bIsMuted)
+{
+	GetVoiceChatUser().SetAudioInputDeviceMuted(bIsMuted);
+}
+
+void FVivoxVoiceChat::SetAudioOutputDeviceMuted(bool bIsMuted)
+{
+	GetVoiceChatUser().SetAudioOutputDeviceMuted(bIsMuted);
+}
+
+bool FVivoxVoiceChat::GetAudioInputDeviceMuted() const
+{
+	return GetVoiceChatUser().GetAudioInputDeviceMuted();
+}
+
+bool FVivoxVoiceChat::GetAudioOutputDeviceMuted() const
+{
+	return GetVoiceChatUser().GetAudioOutputDeviceMuted();
+}
+
+TArray<FString> FVivoxVoiceChat::GetAvailableInputDevices() const
+{
+	return GetVoiceChatUser().GetAvailableOutputDevices();
+}
+
+TArray<FString> FVivoxVoiceChat::GetAvailableOutputDevices() const
+{
+	return GetVoiceChatUser().GetAvailableOutputDevices();
+}
+
+FOnVoiceChatAvailableAudioDevicesChangedDelegate& FVivoxVoiceChat::OnVoiceChatAvailableAudioDevicesChanged()
+{
+	return GetVoiceChatUser().OnVoiceChatAvailableAudioDevicesChanged();
+}
+
+void FVivoxVoiceChat::SetInputDevice(const FString& InputDevice)
+{
+	GetVoiceChatUser().SetInputDevice(InputDevice);
+}
+
+void FVivoxVoiceChat::SetOutputDevice(const FString& OutputDevice)
+{
+	GetVoiceChatUser().SetOutputDevice(OutputDevice);
+}
+
+FString FVivoxVoiceChat::GetInputDevice() const
+{
+	return GetVoiceChatUser().GetInputDevice();
+}
+
+FString FVivoxVoiceChat::GetOutputDevice() const
+{
+	return GetVoiceChatUser().GetOutputDevice();
+}
+
+FString FVivoxVoiceChat::GetDefaultInputDevice() const
+{
+	return GetVoiceChatUser().GetDefaultInputDevice();
+}
+
+FString FVivoxVoiceChat::GetDefaultOutputDevice() const
+{
+	return GetVoiceChatUser().GetDefaultOutputDevice();
+}
+
+void FVivoxVoiceChat::Connect(const FOnVoiceChatConnectCompleteDelegate& Delegate)
+{
+	FVoiceChatResult Result = FVoiceChatResult::CreateSuccess();
+
+	if (!IsInitialized())
+	{
+		Result = VoiceChat::Errors::NotInitialized();
+	}
+	else if (ConnectionState == EConnectionState::Disconnecting)
+	{
+		Result = VoiceChat::Errors::DisconnectInProgress();
+	}
+
+	if (!Result.IsSuccess())
+	{
+		UE_LOG(LogVivoxVoiceChat, Warning, TEXT("Connect %s"), *LexToString(Result));
+
+		Delegate.ExecuteIfBound(Result);
+	}
+	else if (IsConnected())
+	{
+		Delegate.ExecuteIfBound(FVoiceChatResult::CreateSuccess());
+	}
+	else
+	{
+		OnVoiceChatConnectCompleteDelegates.Add(Delegate);
+
+		if (!IsConnecting())
+		{
+			if (VivoxServerUrl.IsEmpty() || VivoxDomain.IsEmpty() || VivoxNamespace.IsEmpty())
+			{
+				UE_LOG(LogVivoxVoiceChat, Warning, TEXT("[VoiceChat.Vivox] ServerUrl, Domain, or Issuer is not set. Vivox voice chat will not work"));
+				Result = VoiceChat::Errors::MissingConfig();
+				Result.ErrorDesc = FString::Printf(TEXT("ServerUrl=[%s] Domain=[%s] Namespace=[%s]"), *VivoxServerUrl, *VivoxDomain, *VivoxNamespace);
+			}
+			else
+			{
+				VivoxClientApi::Uri BackendUri(TCHAR_TO_ANSI(*VivoxServerUrl));
+				VivoxClientApi::VCSStatus Status = VivoxClientConnection.Connect(BackendUri);
+
+				if (Status.IsError())
+				{
+					UE_LOG(LogVivoxVoiceChat, Warning, TEXT("Connect failed: server:%s error:%s (%i)"), ANSI_TO_TCHAR(BackendUri.ToString()), ANSI_TO_TCHAR(Status.ToString()), Status.GetStatusCode());
+					Result = ResultFromVivoxStatus(Status);
+				}
+				else
+				{
+					ConnectionState = EConnectionState::Connecting;
+				}
+			}
+
+			if (!Result.IsSuccess())
+			{
+				ConnectionState = EConnectionState::Disconnected;
+
+				TriggerCompletionDelegates(OnVoiceChatConnectCompleteDelegates, Result);
+			}
+		}
+	}
+}
+
+void FVivoxVoiceChat::Disconnect(const FOnVoiceChatDisconnectCompleteDelegate& Delegate)
+{
+	if (IsConnected())
+	{
+		VivoxClientApi::Uri BackendUri(TCHAR_TO_ANSI(*VivoxServerUrl));
+		VivoxClientConnection.Disconnect(BackendUri);
+
+		ConnectionState = EConnectionState::Disconnecting;
+		OnVoiceChatDisconnectCompleteDelegates.Add(Delegate);
+	}
+	else
+	{
+		Delegate.ExecuteIfBound(FVoiceChatResult::CreateSuccess());
+	}
+}
+
+bool FVivoxVoiceChat::IsConnecting() const
+{
+	return ConnectionState == EConnectionState::Connecting;
+}
+
+bool FVivoxVoiceChat::IsConnected() const
+{
+	return ConnectionState == EConnectionState::Connected;
+}
+
+void FVivoxVoiceChat::Login(FPlatformUserId PlatformId, const FString& PlayerName, const FString& Credentials, const FOnVoiceChatLoginCompleteDelegate& Delegate)
+{
+	GetVoiceChatUser().Login(PlatformId, PlayerName, Credentials, Delegate);
+}
+
+void FVivoxVoiceChat::Logout(const FOnVoiceChatLogoutCompleteDelegate& Delegate)
+{
+	GetVoiceChatUser().Logout(Delegate);
+}
+
+bool FVivoxVoiceChat::IsLoggingIn() const
+{
+	return GetVoiceChatUser().IsLoggingIn();
+}
+
+bool FVivoxVoiceChat::IsLoggedIn() const
+{
+	return GetVoiceChatUser().IsLoggedIn();
+}
+
+FOnVoiceChatLoggedInDelegate& FVivoxVoiceChat::OnVoiceChatLoggedIn()
+{
+	return GetVoiceChatUser().OnVoiceChatLoggedIn();
+}
+
+FOnVoiceChatLoggedOutDelegate& FVivoxVoiceChat::OnVoiceChatLoggedOut()
+{
+	return GetVoiceChatUser().OnVoiceChatLoggedOut();
+}
+
+FString FVivoxVoiceChat::GetLoggedInPlayerName() const
+{
+	return GetVoiceChatUser().GetLoggedInPlayerName();
+}
+
+void FVivoxVoiceChat::BlockPlayers(const TArray<FString>& PlayerNames)
+{
+	GetVoiceChatUser().BlockPlayers(PlayerNames);
+}
+
+void FVivoxVoiceChat::UnblockPlayers(const TArray<FString>& PlayerNames)
+{
+	GetVoiceChatUser().UnblockPlayers(PlayerNames);
+}
+
+void FVivoxVoiceChat::JoinChannel(const FString& ChannelName, const FString& ChannelCredentials, EVoiceChatChannelType ChannelType, const FOnVoiceChatChannelJoinCompleteDelegate& Delegate, TOptional<FVoiceChatChannel3dProperties> Channel3dProperties)
+{
+	GetVoiceChatUser().JoinChannel(ChannelName, ChannelCredentials, ChannelType, Delegate, Channel3dProperties);
+}
+
+void FVivoxVoiceChat::LeaveChannel(const FString& Channel, const FOnVoiceChatChannelLeaveCompleteDelegate& Delegate)
+{
+	GetVoiceChatUser().LeaveChannel(Channel, Delegate);
+}
+
+FOnVoiceChatChannelJoinedDelegate& FVivoxVoiceChat::OnVoiceChatChannelJoined()
+{
+	return GetVoiceChatUser().OnVoiceChatChannelJoined();
+}
+
+FOnVoiceChatChannelExitedDelegate& FVivoxVoiceChat::OnVoiceChatChannelExited()
+{
+	return GetVoiceChatUser().OnVoiceChatChannelExited();
+}
+
+FOnVoiceChatCallStatsUpdatedDelegate& FVivoxVoiceChat::OnVoiceChatCallStatsUpdated()
+{
+	return GetVoiceChatUser().OnVoiceChatCallStatsUpdated();
+}
+
+void FVivoxVoiceChat::Set3DPosition(const FString& ChannelName, const FVector& SpeakerPosition, const FVector& ListenerPosition, const FVector& ListenerForwardDirection, const FVector& ListenerUpDirection)
+{
+	GetVoiceChatUser().Set3DPosition(ChannelName, SpeakerPosition, ListenerPosition, ListenerForwardDirection, ListenerUpDirection);
+}
+
+TArray<FString> FVivoxVoiceChat::GetChannels() const
+{
+	return GetVoiceChatUser().GetChannels();
+}
+
+TArray<FString> FVivoxVoiceChat::GetPlayersInChannel(const FString& ChannelName) const
+{
+	return GetVoiceChatUser().GetPlayersInChannel(ChannelName);
+}
+
+EVoiceChatChannelType FVivoxVoiceChat::GetChannelType(const FString& ChannelName) const
+{
+	return GetVoiceChatUser().GetChannelType(ChannelName);
+}
+
+FOnVoiceChatPlayerAddedDelegate& FVivoxVoiceChat::OnVoiceChatPlayerAdded()
+{
+	return GetVoiceChatUser().OnVoiceChatPlayerAdded();
+}
+
+FOnVoiceChatPlayerRemovedDelegate& FVivoxVoiceChat::OnVoiceChatPlayerRemoved()
+{
+	return GetVoiceChatUser().OnVoiceChatPlayerRemoved();
+}
+
+bool FVivoxVoiceChat::IsPlayerTalking(const FString& PlayerName) const
+{
+	return GetVoiceChatUser().IsPlayerTalking(PlayerName);
+}
+
+FOnVoiceChatPlayerTalkingUpdatedDelegate& FVivoxVoiceChat::OnVoiceChatPlayerTalkingUpdated()
+{
+	return GetVoiceChatUser().OnVoiceChatPlayerTalkingUpdated();
+}
+
+void FVivoxVoiceChat::SetPlayerMuted(const FString& PlayerName, bool bMuted)
+{
+	GetVoiceChatUser().SetPlayerMuted(PlayerName, bMuted);
+}
+
+bool FVivoxVoiceChat::IsPlayerMuted(const FString& PlayerName) const
+{
+	return GetVoiceChatUser().IsPlayerMuted(PlayerName);
+}
+
+FOnVoiceChatPlayerMuteUpdatedDelegate& FVivoxVoiceChat::OnVoiceChatPlayerMuteUpdated()
+{
+	return GetVoiceChatUser().OnVoiceChatPlayerMuteUpdated();
+}
+
+void FVivoxVoiceChat::SetPlayerVolume(const FString& PlayerName, float Volume)
+{
+	GetVoiceChatUser().SetPlayerVolume(PlayerName, Volume);
+}
+
+float FVivoxVoiceChat::GetPlayerVolume(const FString& PlayerName) const
+{
+	return GetVoiceChatUser().GetPlayerVolume(PlayerName);
+}
+
+FOnVoiceChatPlayerVolumeUpdatedDelegate& FVivoxVoiceChat::OnVoiceChatPlayerVolumeUpdated()
+{
+	return GetVoiceChatUser().OnVoiceChatPlayerVolumeUpdated();
+}
+
+void FVivoxVoiceChat::TransmitToAllChannels()
+{
+	GetVoiceChatUser().TransmitToAllChannels();
+}
+
+void FVivoxVoiceChat::TransmitToNoChannels()
+{
+	GetVoiceChatUser().TransmitToNoChannels();
+}
+
+void FVivoxVoiceChat::TransmitToSpecificChannel(const FString& ChannelName)
+{
+	GetVoiceChatUser().TransmitToSpecificChannel(ChannelName);
+}
+
+EVoiceChatTransmitMode FVivoxVoiceChat::GetTransmitMode() const
+{
+	return GetVoiceChatUser().GetTransmitMode();
+}
+
+FString FVivoxVoiceChat::GetTransmitChannel() const
+{
+	return GetVoiceChatUser().GetTransmitChannel();
+}
+
+FDelegateHandle FVivoxVoiceChat::StartRecording(const FOnVoiceChatRecordSamplesAvailableDelegate::FDelegate& Delegate)
+{
+	return GetVoiceChatUser().StartRecording(Delegate);
+}
+
+void FVivoxVoiceChat::StopRecording(FDelegateHandle Handle)
+{
+	GetVoiceChatUser().StopRecording(Handle);
+}
+
+FDelegateHandle FVivoxVoiceChat::RegisterOnVoiceChatAfterCaptureAudioReadDelegate(const FOnVoiceChatAfterCaptureAudioReadDelegate::FDelegate& Delegate)
+{
+	return GetVoiceChatUser().RegisterOnVoiceChatAfterCaptureAudioReadDelegate(Delegate);
+}
+
+void FVivoxVoiceChat::UnregisterOnVoiceChatAfterCaptureAudioReadDelegate(FDelegateHandle Handle)
+{
+	GetVoiceChatUser().UnregisterOnVoiceChatAfterCaptureAudioReadDelegate(Handle);
+}
+
+FDelegateHandle FVivoxVoiceChat::RegisterOnVoiceChatBeforeCaptureAudioSentDelegate(const FOnVoiceChatBeforeCaptureAudioSentDelegate::FDelegate& Delegate)
+{
+	return GetVoiceChatUser().RegisterOnVoiceChatBeforeCaptureAudioSentDelegate(Delegate);
+}
+
+void FVivoxVoiceChat::UnregisterOnVoiceChatBeforeCaptureAudioSentDelegate(FDelegateHandle Handle)
+{
+	GetVoiceChatUser().UnregisterOnVoiceChatBeforeCaptureAudioSentDelegate(Handle);
+}
+
+FDelegateHandle FVivoxVoiceChat::RegisterOnVoiceChatBeforeRecvAudioRenderedDelegate(const FOnVoiceChatBeforeRecvAudioRenderedDelegate::FDelegate& Delegate)
+{
+	return GetVoiceChatUser().RegisterOnVoiceChatBeforeRecvAudioRenderedDelegate(Delegate);
+}
+
+void FVivoxVoiceChat::UnregisterOnVoiceChatBeforeRecvAudioRenderedDelegate(FDelegateHandle Handle)
+{
+	GetVoiceChatUser().UnregisterOnVoiceChatBeforeRecvAudioRenderedDelegate(Handle);
+}
+
+FString FVivoxVoiceChat::InsecureGetLoginToken(const FString& PlayerName)
+{
+	return GetVoiceChatUser().InsecureGetLoginToken(PlayerName);
+}
+
+FString FVivoxVoiceChat::InsecureGetJoinToken(const FString& ChannelName, EVoiceChatChannelType ChannelType, TOptional<FVoiceChatChannel3dProperties> Channel3dProperties)
+{
+	return GetVoiceChatUser().InsecureGetJoinToken(ChannelName, ChannelType, Channel3dProperties);
+}
+
+template <typename TFn, typename... TArgs>
+void FVivoxVoiceChat::DispatchUsingAccountName(const TFn& Fn, const VivoxClientApi::AccountName& AccountName, TArgs&&... Args)
+{
+	for (FVivoxVoiceChatUser* User : VoiceChatUsers)
+	{
+		if (User->LoginSession.AccountName == AccountName)
+		{
+			(User->*Fn)(AccountName, Forward<TArgs>(Args)...);
+		}
+	}
+}
+
+template <typename TFn, typename... TArgs>
+void FVivoxVoiceChat::DispatchAll(const TFn& Fn, TArgs&&... Args)
+{
+	for (FVivoxVoiceChatUser* User : VoiceChatUsers)
+	{
+		(User->*Fn)(Forward<TArgs>(Args)...);
+	}
+}
+
+void FVivoxVoiceChat::InvokeOnUIThread(void (Func)(void* Arg0), void* Arg0)
+{
+	AsyncTask(ENamedThreads::GameThread, [Func, Arg0]()
+	{
+		if (Func)
+		{
+			(*Func)(Arg0);
+		}
+	});
+
+	FEmbeddedCommunication::WakeGameThread();
+}
+
+void FVivoxVoiceChat::onLogStatementEmitted(LogLevel Level, long long NativeMillisecondsSinceEpoch, long ThreadId, const char* LogMessageCStr)
+{
+	FString LogMessage = ANSI_TO_TCHAR(LogMessageCStr);
+	bool bLogMatches = false;
+
+	// Don't spam the log if we receive repeated log messages
+	if (!LastLogMessage.IsEmpty())
+	{
+		bLogMatches = LogMessage == LastLogMessage && Level == LastLogLevel;
+		if (bLogMatches)
+		{
+			LogSpamCount++;
+		}
+
+		const bool bLogSpamCountPowerOfTwo = FMath::IsPowerOfTwo(LogSpamCount);
+		const bool bLogDueToPot = bLogMatches && bLogSpamCountPowerOfTwo;
+		const bool bLogDueToChange = !bLogMatches && LogSpamCount > 0 && !bLogSpamCountPowerOfTwo; // Don't log on change if POT, as we already logged it.
+		if (bLogDueToChange || bLogDueToPot)
+		{
+			if (LastLogLevel == LogLevelError)
+			{
+				UE_LOG(LogVivoxVoiceChat, Warning, TEXT("vivox: Error: %s (seen %d times)"), *LastLogMessage, LogSpamCount);
+			}
+			else
+			{
+				UE_LOG(LogVivoxVoiceChat, Log, TEXT("vivox: %s: %s (seen %d times)"), *GetLogLevelString(LastLogLevel), *LastLogMessage, LogSpamCount);
+			}
+		}
+	}
+
+	if (!bLogMatches)
+	{
+		if (Level == LogLevelError)
+		{
+			UE_LOG(LogVivoxVoiceChat, Warning, TEXT("vivox: Error: %s"), *LogMessage);
+		}
+		else
+		{
+			UE_LOG(LogVivoxVoiceChat, Log, TEXT("vivox: %s: %s"), *GetLogLevelString(Level), *LogMessage);
+		}
+
+		LogSpamCount = 0;
+		LastLogMessage = MoveTemp(LogMessage);
+		LastLogLevel = Level;
+	}
+}
+
+void FVivoxVoiceChat::onConnectCompleted(const VivoxClientApi::Uri& Server)
+{
+	UE_LOG(LogVivoxVoiceChat, Log, TEXT("onConnectCompleted server:%s"), ANSI_TO_TCHAR(Server.ToString()));
+
+	DispatchAll(&FVivoxVoiceChatUser::onConnectCompleted, Server);
+
+	ConnectionState = EConnectionState::Connected;
+
+	TriggerCompletionDelegates(OnVoiceChatConnectCompleteDelegates, FVoiceChatResult::CreateSuccess());
+}
+
+void FVivoxVoiceChat::onConnectFailed(const VivoxClientApi::Uri& Server, const VivoxClientApi::VCSStatus& Status)
+{
+	UE_LOG(LogVivoxVoiceChat, Warning, TEXT("onConnectFailed server:%s error:%s (%i)"), ANSI_TO_TCHAR(Server.ToString()), ANSI_TO_TCHAR(Status.ToString()), Status.GetStatusCode());
+
+	DispatchAll(&FVivoxVoiceChatUser::onConnectFailed, Server, Status);
+
+	ConnectionState = EConnectionState::Disconnected;
+
+	TriggerCompletionDelegates(OnVoiceChatConnectCompleteDelegates, ResultFromVivoxStatus(Status));
+}
+
+void FVivoxVoiceChat::onDisconnected(const VivoxClientApi::Uri& Server, const VivoxClientApi::VCSStatus& Status)
+{
+	if (Status.IsError())
+	{
+		UE_LOG(LogVivoxVoiceChat, Warning, TEXT("onDisconnected server:%s error:%s (%i)"), ANSI_TO_TCHAR(Server.ToString()), ANSI_TO_TCHAR(Status.ToString()), Status.GetStatusCode());
+	}
+	else
+	{
+		UE_LOG(LogVivoxVoiceChat, Log, TEXT("onDisconnected server:%s"), ANSI_TO_TCHAR(Server.ToString()));
+	}
+
+	DispatchAll(&FVivoxVoiceChatUser::onDisconnected, Server, Status);
+
+	EConnectionState PreviousConnectionState = ConnectionState;
+	ConnectionState = EConnectionState::Disconnected;
+
+	if (PreviousConnectionState == EConnectionState::Disconnecting)
+	{
+		TriggerCompletionDelegates(OnVoiceChatDisconnectCompleteDelegates, FVoiceChatResult::CreateSuccess());
+	}
+	else
+	{
+		OnVoiceChatDisconnectedDelegate.Broadcast(ResultFromVivoxStatus(Status));
+	}
+}
+
+void FVivoxVoiceChat::onLoginCompleted(const VivoxClientApi::AccountName& AccountName)
+{
+	DispatchUsingAccountName(&FVivoxVoiceChatUser::onLoginCompleted, AccountName);
+}
+
+void FVivoxVoiceChat::onInvalidLoginCredentials(const VivoxClientApi::AccountName& AccountName)
+{
+	DispatchUsingAccountName(&FVivoxVoiceChatUser::onInvalidLoginCredentials, AccountName);
+}
+
+void FVivoxVoiceChat::onLoginFailed(const VivoxClientApi::AccountName& AccountName, const VivoxClientApi::VCSStatus& Status)
+{
+	DispatchUsingAccountName(&FVivoxVoiceChatUser::onLoginFailed, AccountName, Status);
+}
+
+void FVivoxVoiceChat::onSessionGroupCreated(const VivoxClientApi::AccountName& AccountName, const char* SessionGroupHandle)
+{
+	DispatchUsingAccountName(&FVivoxVoiceChatUser::onSessionGroupCreated, AccountName, SessionGroupHandle);
+}
+
+void FVivoxVoiceChat::onLogoutCompleted(const VivoxClientApi::AccountName& AccountName)
+{
+	DispatchUsingAccountName(&FVivoxVoiceChatUser::onLogoutCompleted, AccountName);
+}
+
+void FVivoxVoiceChat::onLogoutFailed(const VivoxClientApi::AccountName& AccountName, const VivoxClientApi::VCSStatus& Status)
+{
+	DispatchUsingAccountName(&FVivoxVoiceChatUser::onLogoutFailed, AccountName, Status);
+}
+
+void FVivoxVoiceChat::onChannelJoined(const VivoxClientApi::AccountName& AccountName, const VivoxClientApi::Uri& ChannelUri)
+{
+	DispatchUsingAccountName(&FVivoxVoiceChatUser::onChannelJoined, AccountName, ChannelUri);
+}
+
+void FVivoxVoiceChat::onInvalidChannelCredentials(const VivoxClientApi::AccountName& AccountName, const VivoxClientApi::Uri& ChannelUri)
+{
+	DispatchUsingAccountName(&FVivoxVoiceChatUser::onInvalidChannelCredentials, AccountName, ChannelUri);
+}
+
+void FVivoxVoiceChat::onChannelJoinFailed(const VivoxClientApi::AccountName& AccountName, const VivoxClientApi::Uri& ChannelUri, const VivoxClientApi::VCSStatus& Status)
+{
+	DispatchUsingAccountName(&FVivoxVoiceChatUser::onChannelJoinFailed, AccountName, ChannelUri, Status);
+}
+
+void FVivoxVoiceChat::onChannelExited(const VivoxClientApi::AccountName& AccountName, const VivoxClientApi::Uri& ChannelUri, const VivoxClientApi::VCSStatus& ReasonCode)
+{
+	DispatchUsingAccountName(&FVivoxVoiceChatUser::onChannelExited, AccountName, ChannelUri, ReasonCode);
+}
+
+void FVivoxVoiceChat::onCallStatsUpdated(const VivoxClientApi::AccountName& AccountName, vx_call_stats_t& Stats, bool bIsFinal)
+{
+	DispatchUsingAccountName(&FVivoxVoiceChatUser::onCallStatsUpdated, AccountName, Stats, bIsFinal);
+}
+
+void FVivoxVoiceChat::onParticipantAdded(const VivoxClientApi::AccountName& AccountName, const VivoxClientApi::Uri& ChannelUri, const VivoxClientApi::Uri& ParticipantUri, bool bIsLoggedInUser)
+{
+	DispatchUsingAccountName(&FVivoxVoiceChatUser::onParticipantAdded, AccountName, ChannelUri, ParticipantUri, bIsLoggedInUser);
+}
+
+void FVivoxVoiceChat::onParticipantLeft(const VivoxClientApi::AccountName& AccountName, const VivoxClientApi::Uri& ChannelUri, const VivoxClientApi::Uri& ParticipantUri, bool bIsLoggedInUser, ParticipantLeftReason Reason)
+{
+	DispatchUsingAccountName(&FVivoxVoiceChatUser::onParticipantLeft, AccountName, ChannelUri, ParticipantUri, bIsLoggedInUser, Reason);
+}
+
+void FVivoxVoiceChat::onParticipantUpdated(const VivoxClientApi::AccountName& AccountName, const VivoxClientApi::Uri& ChannelUri, const VivoxClientApi::Uri& ParticipantUri, bool bIsLoggedInUser, bool bSpeaking, double MeterEnergy, bool bMutedForAll)
+{
+	DispatchUsingAccountName(&FVivoxVoiceChatUser::onParticipantUpdated, AccountName, ChannelUri, ParticipantUri, bIsLoggedInUser, bSpeaking, MeterEnergy, bMutedForAll);
+}
+
+void FVivoxVoiceChat::onAvailableAudioDevicesChanged()
+{
+	DispatchAll(&FVivoxVoiceChatUser::onAvailableAudioDevicesChanged);
+}
+void FVivoxVoiceChat::onOperatingSystemChosenAudioInputDeviceChanged(const VivoxClientApi::AudioDeviceId& DeviceId)
+{
+	DispatchAll(&FVivoxVoiceChatUser::onOperatingSystemChosenAudioInputDeviceChanged, DeviceId);
+}
+
+void FVivoxVoiceChat::onSetApplicationChosenAudioInputDeviceCompleted(const VivoxClientApi::AccountName& AccountName, const VivoxClientApi::AudioDeviceId& DeviceId)
+{
+	DispatchUsingAccountName(&FVivoxVoiceChatUser::onSetApplicationChosenAudioInputDeviceCompleted, AccountName, DeviceId);
+}
+
+void FVivoxVoiceChat::onSetApplicationChosenAudioInputDeviceFailed(const VivoxClientApi::AccountName& AccountName, const VivoxClientApi::AudioDeviceId& DeviceId, const VivoxClientApi::VCSStatus& Status)
+{
+	DispatchUsingAccountName(&FVivoxVoiceChatUser::onSetApplicationChosenAudioInputDeviceFailed, AccountName, DeviceId, Status);
+}
+
+void FVivoxVoiceChat::onOperatingSystemChosenAudioOutputDeviceChanged(const VivoxClientApi::AudioDeviceId& DeviceId)
+{
+	DispatchAll(&FVivoxVoiceChatUser::onOperatingSystemChosenAudioOutputDeviceChanged, DeviceId);
+}
+
+void FVivoxVoiceChat::onSetApplicationChosenAudioOutputDeviceCompleted(const VivoxClientApi::AccountName& AccountName, const VivoxClientApi::AudioDeviceId& DeviceId)
+{
+	DispatchUsingAccountName(&FVivoxVoiceChatUser::onSetApplicationChosenAudioOutputDeviceCompleted, AccountName, DeviceId);
+}
+
+void FVivoxVoiceChat::onSetApplicationChosenAudioOutputDeviceFailed(const VivoxClientApi::AccountName& AccountName, const VivoxClientApi::AudioDeviceId& DeviceId, const VivoxClientApi::VCSStatus& Status)
+{
+	DispatchUsingAccountName(&FVivoxVoiceChatUser::onSetApplicationChosenAudioOutputDeviceFailed, AccountName, DeviceId, Status);
+}
+
+void FVivoxVoiceChat::onSetParticipantAudioOutputDeviceVolumeForMeCompleted(const VivoxClientApi::AccountName& AccountName, const VivoxClientApi::Uri& TargetUser, const VivoxClientApi::Uri& ChannelUri, int Volume)
+{
+	DispatchUsingAccountName(&FVivoxVoiceChatUser::onSetParticipantAudioOutputDeviceVolumeForMeCompleted, AccountName, TargetUser, ChannelUri, Volume);
+}
+
+void FVivoxVoiceChat::onSetParticipantAudioOutputDeviceVolumeForMeFailed(const VivoxClientApi::AccountName& AccountName, const VivoxClientApi::Uri& TargetUser, const VivoxClientApi::Uri& ChannelUri, int Volume, const VivoxClientApi::VCSStatus& Status)
+{
+	DispatchUsingAccountName(&FVivoxVoiceChatUser::onSetParticipantAudioOutputDeviceVolumeForMeFailed, AccountName, TargetUser, ChannelUri, Volume, Status);
+}
+
+void FVivoxVoiceChat::onSetChannelAudioOutputDeviceVolumeCompleted(const VivoxClientApi::AccountName& AccountName, const VivoxClientApi::Uri& ChannelUri, int Volume)
+{
+	DispatchUsingAccountName(&FVivoxVoiceChatUser::onSetChannelAudioOutputDeviceVolumeCompleted, AccountName, ChannelUri, Volume);
+}
+
+void FVivoxVoiceChat::onSetChannelAudioOutputDeviceVolumeFailed(const VivoxClientApi::AccountName& AccountName, const VivoxClientApi::Uri& ChannelUri, int Volume, const VivoxClientApi::VCSStatus& Status)
+{
+	DispatchUsingAccountName(&FVivoxVoiceChatUser::onSetChannelAudioOutputDeviceVolumeFailed, AccountName, ChannelUri, Volume, Status);
+}
+
+void FVivoxVoiceChat::onSetParticipantMutedForMeCompleted(const VivoxClientApi::AccountName& AccountName, const VivoxClientApi::Uri& Target, const VivoxClientApi::Uri& ChannelUri, bool bMuted)
+{
+	DispatchUsingAccountName(&FVivoxVoiceChatUser::onSetParticipantMutedForMeCompleted, AccountName, Target, ChannelUri, bMuted);
+}
+
+void FVivoxVoiceChat::onSetParticipantMutedForMeFailed(const VivoxClientApi::AccountName& AccountName, const VivoxClientApi::Uri& Target, const VivoxClientApi::Uri& ChannelUri, bool bMuted, const VivoxClientApi::VCSStatus& Status)
+{
+	DispatchUsingAccountName(&FVivoxVoiceChatUser::onSetParticipantMutedForMeFailed, AccountName, Target, ChannelUri, bMuted, Status);
+}
+
+void FVivoxVoiceChat::onSetChannelTransmissionToSpecificChannelCompleted(const VivoxClientApi::AccountName& AccountName, const VivoxClientApi::Uri& ChannelUri)
+{
+	DispatchUsingAccountName(&FVivoxVoiceChatUser::onSetChannelTransmissionToSpecificChannelCompleted, AccountName, ChannelUri);
+}
+
+void FVivoxVoiceChat::onSetChannelTransmissionToSpecificChannelFailed(const VivoxClientApi::AccountName& AccountName, const VivoxClientApi::Uri& ChannelUri, const VivoxClientApi::VCSStatus& Status)
+{
+	DispatchUsingAccountName(&FVivoxVoiceChatUser::onSetChannelTransmissionToSpecificChannelFailed, AccountName, ChannelUri, Status);
+}
+
+void FVivoxVoiceChat::onSetChannelTransmissionToAllCompleted(const VivoxClientApi::AccountName& AccountName)
+{
+	DispatchUsingAccountName(&FVivoxVoiceChatUser::onSetChannelTransmissionToAllCompleted, AccountName);
+}
+
+void FVivoxVoiceChat::onSetChannelTransmissionToAllFailed(const VivoxClientApi::AccountName& AccountName, const VivoxClientApi::VCSStatus& Status)
+{
+	DispatchUsingAccountName(&FVivoxVoiceChatUser::onSetChannelTransmissionToAllFailed, AccountName, Status);
+}
+
+void FVivoxVoiceChat::onSetChannelTransmissionToNoneCompleted(const VivoxClientApi::AccountName& AccountName)
+{
+	DispatchUsingAccountName(&FVivoxVoiceChatUser::onSetChannelTransmissionToNoneCompleted, AccountName);
+}
+
+void FVivoxVoiceChat::onSetChannelTransmissionToNoneFailed(const VivoxClientApi::AccountName& AccountName, const VivoxClientApi::VCSStatus& Status)
+{
+	DispatchUsingAccountName(&FVivoxVoiceChatUser::onSetChannelTransmissionToNoneFailed, AccountName, Status);
+}
+
+void FVivoxVoiceChat::onAudioUnitStarted(const char* SessionGroupHandle, const VivoxClientApi::Uri& InitialTargetUri)
+{
+	FScopeLock Lock(&VoiceChatUsersCriticalSection);
+	DispatchAll(&FVivoxVoiceChatUser::onAudioUnitStarted, SessionGroupHandle, InitialTargetUri);
+}
+
+void FVivoxVoiceChat::onAudioUnitStopped(const char* SessionGroupHandle, const VivoxClientApi::Uri& InitialTargetUri)
+{
+	FScopeLock Lock(&VoiceChatUsersCriticalSection);
+	DispatchAll(&FVivoxVoiceChatUser::onAudioUnitStopped, SessionGroupHandle, InitialTargetUri);
+}
+
+void FVivoxVoiceChat::onAudioUnitAfterCaptureAudioRead(const char* SessionGroupHandle, const VivoxClientApi::Uri& InitialTargetUri, short* PcmFrames, int PcmFrameCount, int AudioFrameRate, int ChannelsPerFrame)
+{
+	FScopeLock Lock(&VoiceChatUsersCriticalSection);
+	DispatchAll(&FVivoxVoiceChatUser::onAudioUnitAfterCaptureAudioRead, SessionGroupHandle, InitialTargetUri, PcmFrames, PcmFrameCount, AudioFrameRate, ChannelsPerFrame);
+}
+
+void FVivoxVoiceChat::onAudioUnitBeforeCaptureAudioSent(const char* SessionGroupHandle, const VivoxClientApi::Uri& InitialTargetUri, short* PcmFrames, int PcmFrameCount, int AudioFrameRate, int ChannelsPerFrame, bool bSpeaking)
+{
+	FScopeLock Lock(&VoiceChatUsersCriticalSection);
+	DispatchAll(&FVivoxVoiceChatUser::onAudioUnitBeforeCaptureAudioSent, SessionGroupHandle, InitialTargetUri, PcmFrames, PcmFrameCount, AudioFrameRate, ChannelsPerFrame, bSpeaking);
+}
+
+void FVivoxVoiceChat::onAudioUnitBeforeRecvAudioRendered(const char* SessionGroupHandle, const VivoxClientApi::Uri& InitialTargetUri, short* PcmFrames, int PcmFrameCount, int AudioFrameRate, int ChannelsPerFrame, bool bSilence)
+{
+	FScopeLock Lock(&VoiceChatUsersCriticalSection);
+	DispatchAll(&FVivoxVoiceChatUser::onAudioUnitBeforeRecvAudioRendered, SessionGroupHandle, InitialTargetUri, PcmFrames, PcmFrameCount, AudioFrameRate, ChannelsPerFrame, bSilence);
+}
+
+void FVivoxVoiceChat::RegisterVoiceChatUser(FVivoxVoiceChatUser* User)
+{
+	FScopeLock Lock(&VoiceChatUsersCriticalSection);
+	VoiceChatUsers.Add(User);
+}
+
+void FVivoxVoiceChat::UnregisterVoiceChatUser(FVivoxVoiceChatUser* User)
+{
+	FScopeLock Lock(&VoiceChatUsersCriticalSection);
+	VoiceChatUsers.Remove(User);
+}
diff --git a/Engine/Plugins/Online/VoiceChat/VivoxVoiceChat/Source/Public/VivoxVoiceChat.h b/Engine/Plugins/Online/VoiceChat/VivoxVoiceChat/Source/Public/VivoxVoiceChat.h
--- a/Engine/Plugins/Online/VoiceChat/VivoxVoiceChat/Source/Public/VivoxVoiceChat.h
+++ b/Engine/Plugins/Online/VoiceChat/VivoxVoiceChat/Source/Public/VivoxVoiceChat.h
@@ -31,16 +31,15 @@
 	static FOnAudioUnitCaptureDeviceStatusChanged OnAudioUnitCaptureDeviceStatusChanged;
 };
 
-class FVivoxVoiceChat : public FSelfRegisteringExec, public IVoiceChat, protected VivoxClientApi::DebugClientApiEventHandler
+class FVivoxVoiceChat;
+
+class FVivoxVoiceChatUser : public IVoiceChatUser
 {
 public:
-	FVivoxVoiceChat();
-	virtual ~FVivoxVoiceChat();
+	FVivoxVoiceChatUser(FVivoxVoiceChat& InVivoxVoiceChat);
+	~FVivoxVoiceChatUser();
 
-	// ~Begin IVoiceChat Interface
-	virtual bool Initialize() override;
-	virtual bool Uninitialize() override;
-	virtual bool IsInitialized() const override;
+	// ~Begin IVoiceChatUser Interface
 	virtual void SetSetting(const FString& Name, const FString& Value) override;
 	virtual FString GetSetting(const FString& Name) override;
 	virtual void SetAudioInputVolume(float Volume) override;
@@ -64,9 +63,9 @@
 	virtual void Disconnect(const FOnVoiceChatDisconnectCompleteDelegate& Delegate) override;
 	virtual bool IsConnecting() const override;
 	virtual bool IsConnected() const override;
-	virtual FOnVoiceChatConnectedDelegate& OnVoiceChatConnected() override { return OnVoiceChatConnectedDelegate; }
-	virtual FOnVoiceChatDisconnectedDelegate& OnVoiceChatDisconnected() override { return OnVoiceChatDisconnectedDelegate; }
-	virtual FOnVoiceChatReconnectedDelegate& OnVoiceChatReconnected() override { return OnVoiceChatReconnectedDelegate; }
+	virtual FOnVoiceChatConnectedDelegate& OnVoiceChatConnected() override;
+	virtual FOnVoiceChatDisconnectedDelegate& OnVoiceChatDisconnected() override;
+	virtual FOnVoiceChatReconnectedDelegate& OnVoiceChatReconnected() override;
 	virtual void Login(FPlatformUserId PlatformId, const FString& PlayerName, const FString& Credentials, const FOnVoiceChatLoginCompleteDelegate& Delegate) override;
 	virtual void Logout(const FOnVoiceChatLogoutCompleteDelegate& Delegate) override;
 	virtual bool IsLoggingIn() const override;
@@ -113,53 +112,54 @@
 	// ~End IVoiceChat Interface
 
 protected:
+	bool IsInitialized();
+	bool Exec(UWorld* InWorld, const TCHAR* Cmd, FOutputDevice& Ar);
+	
+	friend class FVivoxVoiceChat;
 	// ~Begin DebugClientApiEventHandler Interface
-	virtual void InvokeOnUIThread(void (Func)(void* Arg0), void* Arg0) override;
-	virtual void onLogStatementEmitted(LogLevel Level, long long NativeMillisecondsSinceEpoch, long ThreadId, const char* LogMessage) override;
-	virtual void onConnectCompleted(const VivoxClientApi::Uri& Server) override;
-	virtual void onConnectFailed(const VivoxClientApi::Uri& Server, const VivoxClientApi::VCSStatus& Status) override;
-	virtual void onDisconnected(const VivoxClientApi::Uri& Server, const VivoxClientApi::VCSStatus& Status) override;
-	virtual void onLoginCompleted(const VivoxClientApi::AccountName& AccountName) override;
-	virtual void onInvalidLoginCredentials(const VivoxClientApi::AccountName& AccountName) override;
-	virtual void onLoginFailed(const VivoxClientApi::AccountName& AccountName, const VivoxClientApi::VCSStatus& Status) override;
-	virtual void onLogoutCompleted(const VivoxClientApi::AccountName& AccountName) override;
-	virtual void onLogoutFailed(const VivoxClientApi::AccountName& AccountName, const VivoxClientApi::VCSStatus& Status) override;
-	virtual void onChannelJoined(const VivoxClientApi::AccountName& AccountName, const VivoxClientApi::Uri& ChannelUri) override;
-	virtual void onInvalidChannelCredentials(const VivoxClientApi::AccountName& AccountName, const VivoxClientApi::Uri& ChannelUri) override;
-	virtual void onChannelJoinFailed(const VivoxClientApi::AccountName& AccountName, const VivoxClientApi::Uri& ChannelUri, const VivoxClientApi::VCSStatus& Status) override;
-	virtual void onChannelExited(const VivoxClientApi::AccountName& AccountName, const VivoxClientApi::Uri& ChannelUri, const VivoxClientApi::VCSStatus& ReasonCode) override;
-	virtual void onCallStatsUpdated(const VivoxClientApi::AccountName& AccountName, vx_call_stats_t& Stats, bool bIsFinal) override;
-	virtual void onParticipantAdded(const VivoxClientApi::AccountName& AccountName, const VivoxClientApi::Uri& ChannelUri, const VivoxClientApi::Uri& ParticipantUri, bool bIsLoggedInUser) override;
-	virtual void onParticipantLeft(const VivoxClientApi::AccountName& AccountName, const VivoxClientApi::Uri& ChannelUri, const VivoxClientApi::Uri& ParticipantUri, bool bIsLoggedInUser, ParticipantLeftReason Reason) override;
-	virtual void onParticipantUpdated(const VivoxClientApi::AccountName& AccountName, const VivoxClientApi::Uri& ChannelUri, const VivoxClientApi::Uri& ParticipantUri, bool bIsLoggedInUser, bool bSpeaking, double MeterEnergy, bool bMutedForAll) override;
-	virtual void onAvailableAudioDevicesChanged() override;
-	virtual void onOperatingSystemChosenAudioInputDeviceChanged(const VivoxClientApi::AudioDeviceId& DeviceId) override;
-	virtual void onSetApplicationChosenAudioInputDeviceCompleted(const VivoxClientApi::AudioDeviceId& DeviceId) override;
-	virtual void onSetApplicationChosenAudioInputDeviceFailed(const VivoxClientApi::AudioDeviceId& DeviceId, const VivoxClientApi::VCSStatus& Status) override;
-	virtual void onOperatingSystemChosenAudioOutputDeviceChanged(const VivoxClientApi::AudioDeviceId& DeviceId) override;
-	virtual void onSetApplicationChosenAudioOutputDeviceCompleted(const VivoxClientApi::AudioDeviceId& DeviceId) override;
-	virtual void onSetApplicationChosenAudioOutputDeviceFailed(const VivoxClientApi::AudioDeviceId& DeviceId, const VivoxClientApi::VCSStatus& Status) override;
-	virtual void onSetParticipantAudioOutputDeviceVolumeForMeCompleted(const VivoxClientApi::AccountName& AccountName, const VivoxClientApi::Uri& TargetUser, const VivoxClientApi::Uri& ChannelUri, int Volume) override;
-	virtual void onSetParticipantAudioOutputDeviceVolumeForMeFailed(const VivoxClientApi::AccountName& AccountName, const VivoxClientApi::Uri& TargetUser, const VivoxClientApi::Uri& ChannelUri, int Volume, const VivoxClientApi::VCSStatus& Status) override;
-	virtual void onSetChannelAudioOutputDeviceVolumeCompleted(const VivoxClientApi::AccountName& AccountName, const VivoxClientApi::Uri& ChannelUri, int Volume) override;
-	virtual void onSetChannelAudioOutputDeviceVolumeFailed(const VivoxClientApi::AccountName& AccountName, const VivoxClientApi::Uri& ChannelUri, int Volume, const VivoxClientApi::VCSStatus& Status) override;
-	virtual void onSetParticipantMutedForMeCompleted(const VivoxClientApi::AccountName& AccountName, const VivoxClientApi::Uri& Target, const VivoxClientApi::Uri& ChannelUri, bool bMuted) override;
-	virtual void onSetParticipantMutedForMeFailed(const VivoxClientApi::AccountName& AccountName, const VivoxClientApi::Uri& Target, const VivoxClientApi::Uri& ChannelUri, bool bMuted, const VivoxClientApi::VCSStatus& Status) override;
-	virtual void onSetChannelTransmissionToSpecificChannelCompleted(const VivoxClientApi::AccountName& AccountName, const VivoxClientApi::Uri& ChannelUri) override;
-	virtual void onSetChannelTransmissionToSpecificChannelFailed(const VivoxClientApi::AccountName& AccountName, const VivoxClientApi::Uri& ChannelUri, const VivoxClientApi::VCSStatus& Status) override;
-	virtual void onSetChannelTransmissionToAllCompleted(const VivoxClientApi::AccountName& AccountName) override;
-	virtual void onSetChannelTransmissionToAllFailed(const VivoxClientApi::AccountName& AccountName, const VivoxClientApi::VCSStatus& Status) override;
-	virtual void onSetChannelTransmissionToNoneCompleted(const VivoxClientApi::AccountName& AccountName) override;
-	virtual void onSetChannelTransmissionToNoneFailed(const VivoxClientApi::AccountName& AccountName, const VivoxClientApi::VCSStatus& Status) override;
-	virtual void onAudioUnitStarted(const VivoxClientApi::Uri& InitialTargetUri) override;
-	virtual void onAudioUnitStopped(const VivoxClientApi::Uri& InitialTargetUri) override;
-	virtual void onAudioUnitAfterCaptureAudioRead(const VivoxClientApi::Uri& InitialTargetUri, short* PcmFrames, int PcmFrameCount, int AudioFrameRate, int ChannelsPerFrame) override;
-	virtual void onAudioUnitBeforeCaptureAudioSent(const VivoxClientApi::Uri& InitialTargetUri, short* PcmFrames, int PcmFrameCount, int AudioFrameRate, int ChannelsPerFrame, bool bSpeaking) override;
-	virtual void onAudioUnitBeforeRecvAudioRendered(const VivoxClientApi::Uri& InitialTargetUri, short* PcmFrames, int PcmFrameCount, int AudioFrameRate, int ChannelsPerFrame, bool bSilence) override;
+	void onConnectCompleted(const VivoxClientApi::Uri& Server);
+	void onConnectFailed(const VivoxClientApi::Uri& Server, const VivoxClientApi::VCSStatus& Status);
+	void onDisconnected(const VivoxClientApi::Uri& Server, const VivoxClientApi::VCSStatus& Status);
+	void onLoginCompleted(const VivoxClientApi::AccountName& AccountName);
+	void onInvalidLoginCredentials(const VivoxClientApi::AccountName& AccountName);
+	void onLoginFailed(const VivoxClientApi::AccountName& AccountName, const VivoxClientApi::VCSStatus& Status);
+	void onLogoutCompleted(const VivoxClientApi::AccountName& AccountName);
+	void onLogoutFailed(const VivoxClientApi::AccountName& AccountName, const VivoxClientApi::VCSStatus& Status);
+	void onSessionGroupCreated(const VivoxClientApi::AccountName& AccountName, const char* SessionGroupHandle);
+	void onChannelJoined(const VivoxClientApi::AccountName& AccountName, const VivoxClientApi::Uri& ChannelUri);
+	void onInvalidChannelCredentials(const VivoxClientApi::AccountName& AccountName, const VivoxClientApi::Uri& ChannelUri);
+	void onChannelJoinFailed(const VivoxClientApi::AccountName& AccountName, const VivoxClientApi::Uri& ChannelUri, const VivoxClientApi::VCSStatus& Status);
+	void onChannelExited(const VivoxClientApi::AccountName& AccountName, const VivoxClientApi::Uri& ChannelUri, const VivoxClientApi::VCSStatus& ReasonCode);
+	void onCallStatsUpdated(const VivoxClientApi::AccountName& AccountName, vx_call_stats_t& Stats, bool bIsFinal);
+	void onParticipantAdded(const VivoxClientApi::AccountName& AccountName, const VivoxClientApi::Uri& ChannelUri, const VivoxClientApi::Uri& ParticipantUri, bool bIsLoggedInUser);
+	void onParticipantLeft(const VivoxClientApi::AccountName& AccountName, const VivoxClientApi::Uri& ChannelUri, const VivoxClientApi::Uri& ParticipantUri, bool bIsLoggedInUser, VivoxClientApi::IClientApiEventHandler::ParticipantLeftReason Reason);
+	void onParticipantUpdated(const VivoxClientApi::AccountName& AccountName, const VivoxClientApi::Uri& ChannelUri, const VivoxClientApi::Uri& ParticipantUri, bool bIsLoggedInUser, bool bSpeaking, double MeterEnergy, bool bMutedForAll);
+	virtual void onAvailableAudioDevicesChanged();
+	void onOperatingSystemChosenAudioInputDeviceChanged(const VivoxClientApi::AudioDeviceId& DeviceId);
+	void onSetApplicationChosenAudioInputDeviceCompleted(const VivoxClientApi::AccountName& AccountName, const VivoxClientApi::AudioDeviceId& DeviceId);
+	void onSetApplicationChosenAudioInputDeviceFailed(const VivoxClientApi::AccountName& AccountName, const VivoxClientApi::AudioDeviceId& DeviceId, const VivoxClientApi::VCSStatus& Status);
+	void onOperatingSystemChosenAudioOutputDeviceChanged(const VivoxClientApi::AudioDeviceId& DeviceId);
+	void onSetApplicationChosenAudioOutputDeviceCompleted(const VivoxClientApi::AccountName& AccountName, const VivoxClientApi::AudioDeviceId& DeviceId);
+	void onSetApplicationChosenAudioOutputDeviceFailed(const VivoxClientApi::AccountName& AccountName, const VivoxClientApi::AudioDeviceId& DeviceId, const VivoxClientApi::VCSStatus& Status);
+	void onSetParticipantAudioOutputDeviceVolumeForMeCompleted(const VivoxClientApi::AccountName& AccountName, const VivoxClientApi::Uri& TargetUser, const VivoxClientApi::Uri& ChannelUri, int Volume);
+	void onSetParticipantAudioOutputDeviceVolumeForMeFailed(const VivoxClientApi::AccountName& AccountName, const VivoxClientApi::Uri& TargetUser, const VivoxClientApi::Uri& ChannelUri, int Volume, const VivoxClientApi::VCSStatus& Status);
+	void onSetChannelAudioOutputDeviceVolumeCompleted(const VivoxClientApi::AccountName& AccountName, const VivoxClientApi::Uri& ChannelUri, int Volume);
+	void onSetChannelAudioOutputDeviceVolumeFailed(const VivoxClientApi::AccountName& AccountName, const VivoxClientApi::Uri& ChannelUri, int Volume, const VivoxClientApi::VCSStatus& Status);
+	void onSetParticipantMutedForMeCompleted(const VivoxClientApi::AccountName& AccountName, const VivoxClientApi::Uri& Target, const VivoxClientApi::Uri& ChannelUri, bool bMuted);
+	void onSetParticipantMutedForMeFailed(const VivoxClientApi::AccountName& AccountName, const VivoxClientApi::Uri& Target, const VivoxClientApi::Uri& ChannelUri, bool bMuted, const VivoxClientApi::VCSStatus& Status);
+	void onSetChannelTransmissionToSpecificChannelCompleted(const VivoxClientApi::AccountName& AccountName, const VivoxClientApi::Uri& ChannelUri);
+	void onSetChannelTransmissionToSpecificChannelFailed(const VivoxClientApi::AccountName& AccountName, const VivoxClientApi::Uri& ChannelUri, const VivoxClientApi::VCSStatus& Status);
+	void onSetChannelTransmissionToAllCompleted(const VivoxClientApi::AccountName& AccountName);
+	void onSetChannelTransmissionToAllFailed(const VivoxClientApi::AccountName& AccountName, const VivoxClientApi::VCSStatus& Status);
+	void onSetChannelTransmissionToNoneCompleted(const VivoxClientApi::AccountName& AccountName);
+	void onSetChannelTransmissionToNoneFailed(const VivoxClientApi::AccountName& AccountName, const VivoxClientApi::VCSStatus& Status);
+	void onAudioUnitStarted(const char* SessionGroupHandle, const VivoxClientApi::Uri& InitialTargetUri);
+	void onAudioUnitStopped(const char* SessionGroupHandle, const VivoxClientApi::Uri& InitialTargetUri);
+	void onAudioUnitAfterCaptureAudioRead(const char* SessionGroupHandle, const VivoxClientApi::Uri& InitialTargetUri, short* PcmFrames, int PcmFrameCount, int AudioFrameRate, int ChannelsPerFrame);
+	void onAudioUnitBeforeCaptureAudioSent(const char* SessionGroupHandle, const VivoxClientApi::Uri& InitialTargetUri, short* PcmFrames, int PcmFrameCount, int AudioFrameRate, int ChannelsPerFrame, bool bSpeaking);
+	void onAudioUnitBeforeRecvAudioRendered(const char* SessionGroupHandle, const VivoxClientApi::Uri& InitialTargetUri, short* PcmFrames, int PcmFrameCount, int AudioFrameRate, int ChannelsPerFrame, bool bSilence);
 	// ~End DebugClientApiEventHandler Interface
 
-	virtual void SetVivoxSdkConfigHints(vx_sdk_config_t& Hints);
-
 	struct FParticipant
 	{
 		FString PlayerName;
@@ -213,36 +213,12 @@
 	AudioOptions AudioInputOptions;
 	AudioOptions AudioOutputOptions;
 
-	enum class EConnectionState
-	{
-		Disconnected,
-		Disconnecting,
-		Connecting,
-		Connected
-	};
+	VivoxClientApi::ClientConnection& VivoxClientConnection;
 
-	VivoxClientApi::ClientConnection VivoxClientConnection;
-
-	bool bInitialized;
-	EConnectionState ConnectionState;
 	FLoginSession LoginSession;
 
-	// Settings
-	FString VivoxServerUrl;
-	FString VivoxDomain;
-	FString VivoxIssuer;
-	FString VivoxNamespace;
-	FString VivoxInsecureSecret;
-	EVoiceChatAttenuationModel AttenuationModel;
-	int MinDistance;
-	int MaxDistance;
-	float Rolloff;
-
 	// Delegates
 	FOnVoiceChatAvailableAudioDevicesChangedDelegate OnVoiceChatAvailableAudioDevicesChangedDelegate;
-	FOnVoiceChatConnectedDelegate OnVoiceChatConnectedDelegate;
-	FOnVoiceChatDisconnectedDelegate OnVoiceChatDisconnectedDelegate;
-	FOnVoiceChatReconnectedDelegate OnVoiceChatReconnectedDelegate;
 	FOnVoiceChatLoggedInDelegate OnVoiceChatLoggedInDelegate;
 	FOnVoiceChatLoggedOutDelegate OnVoiceChatLoggedOutDelegate;
 	FOnVoiceChatChannelJoinedDelegate OnVoiceChatChannelJoinedDelegate;
@@ -265,21 +241,9 @@
 	FOnVoiceChatBeforeRecvAudioRenderedDelegate OnVoiceChatBeforeRecvAudioRenderedDelegate;
 
 	// Completion delegates
-	TArray<FOnVoiceChatConnectCompleteDelegate> OnVoiceChatConnectCompleteDelegates;
-	TArray<FOnVoiceChatDisconnectCompleteDelegate> OnVoiceChatDisconnectCompleteDelegates;
 	FOnVoiceChatLoginCompleteDelegate OnVoiceChatLoginCompleteDelegate;
 	FOnVoiceChatLogoutCompleteDelegate OnVoiceChatLogoutCompleteDelegate;
 
-	VivoxClientApi::AccountName CreateAccountName(const FString& PlayerName);
-	FString GetPlayerNameFromAccountName(const VivoxClientApi::AccountName& AccountName);
-
-	VivoxClientApi::Uri CreateUserUri(const FString& PlayerName);
-	FString GetPlayerNameFromUri(const VivoxClientApi::Uri& UserUri);
-
-	VivoxClientApi::Uri CreateChannelUri(const FString& ChannelName, EVoiceChatChannelType ChannelType, TOptional<FVoiceChatChannel3dProperties> Channel3dProperties = TOptional<FVoiceChatChannel3dProperties>());
-	FString GetChannelNameFromUri(const VivoxClientApi::Uri& ChannelUri);
-	EVoiceChatChannelType GetChannelTypeFromUri(const VivoxClientApi::Uri& ChannelUri);
-
 	FParticipant& GetParticipant(const FString& PlayerName);
 	const FParticipant& GetParticipant(const FString& PlayerName) const;
 
@@ -292,16 +256,213 @@
 	void ApplyAudioInputOptions();
 	void ApplyAudioOutputOptions();
 
+	static FString ToString(FLoginSession::EState State);
+	static FString ToString(FChannelSession::EState State);
+
+	FVivoxVoiceChat& VivoxVoiceChat;
+	FString SessionGroup;
+};
+
+class FVivoxVoiceChat : public FSelfRegisteringExec, public IVoiceChat, protected VivoxClientApi::DebugClientApiEventHandler
+{
+public:
+	FVivoxVoiceChat();
+	virtual ~FVivoxVoiceChat();
+
+	// IVoiceChat Interface
+	virtual bool Initialize() override;
+	virtual bool Uninitialize() override;
+	virtual bool IsInitialized() const override;
+	virtual IVoiceChatUser* CreateUser() override;
+
+	// IVoiceChatUser
+	virtual void SetSetting(const FString& Name, const FString& Value) override;
+	virtual FString GetSetting(const FString& Name) override;
+	virtual void SetAudioInputVolume(float Volume) override;
+	virtual void SetAudioOutputVolume(float Volume) override;
+	virtual float GetAudioInputVolume() const override;
+	virtual float GetAudioOutputVolume() const override;
+	virtual void SetAudioInputDeviceMuted(bool bIsMuted) override;
+	virtual void SetAudioOutputDeviceMuted(bool bIsMuted) override;
+	virtual bool GetAudioInputDeviceMuted() const override;
+	virtual bool GetAudioOutputDeviceMuted() const override;
+	virtual TArray<FString> GetAvailableInputDevices() const override;
+	virtual TArray<FString> GetAvailableOutputDevices() const override;
+	virtual FOnVoiceChatAvailableAudioDevicesChangedDelegate& OnVoiceChatAvailableAudioDevicesChanged() override;
+	virtual void SetInputDevice(const FString& InputDevice) override;
+	virtual void SetOutputDevice(const FString& OutputDevice) override;
+	virtual FString GetInputDevice() const override;
+	virtual FString GetOutputDevice() const override;
+	virtual FString GetDefaultInputDevice() const override;
+	virtual FString GetDefaultOutputDevice() const override;
+	virtual void Connect(const FOnVoiceChatConnectCompleteDelegate& Delegate) override;
+	virtual void Disconnect(const FOnVoiceChatDisconnectCompleteDelegate& Delegate) override;
+	virtual bool IsConnecting() const override;
+	virtual bool IsConnected() const override;
+	virtual FOnVoiceChatConnectedDelegate& OnVoiceChatConnected() override { return OnVoiceChatConnectedDelegate; }
+	virtual FOnVoiceChatDisconnectedDelegate& OnVoiceChatDisconnected() override { return OnVoiceChatDisconnectedDelegate; }
+	virtual FOnVoiceChatReconnectedDelegate& OnVoiceChatReconnected() override { return OnVoiceChatReconnectedDelegate; }
+	virtual void Login(FPlatformUserId PlatformId, const FString& PlayerName, const FString& Credentials, const FOnVoiceChatLoginCompleteDelegate& Delegate) override;
+	virtual void Logout(const FOnVoiceChatLogoutCompleteDelegate& Delegate) override;
+	virtual bool IsLoggingIn() const override;
+	virtual bool IsLoggedIn() const override;
+	virtual FOnVoiceChatLoggedInDelegate& OnVoiceChatLoggedIn() override;
+	virtual FOnVoiceChatLoggedOutDelegate& OnVoiceChatLoggedOut() override;
+	virtual FString GetLoggedInPlayerName() const override;
+	virtual void BlockPlayers(const TArray<FString>& PlayerNames) override;
+	virtual void UnblockPlayers(const TArray<FString>& PlayerNames) override;
+	virtual void JoinChannel(const FString& ChannelName, const FString& ChannelCredentials, EVoiceChatChannelType ChannelType, const FOnVoiceChatChannelJoinCompleteDelegate& Delegate, TOptional<FVoiceChatChannel3dProperties> Channel3dProperties = TOptional<FVoiceChatChannel3dProperties>()) override;
+	virtual void LeaveChannel(const FString& Channel, const FOnVoiceChatChannelLeaveCompleteDelegate& Delegate) override;
+	virtual FOnVoiceChatChannelJoinedDelegate& OnVoiceChatChannelJoined() override;
+	virtual FOnVoiceChatChannelExitedDelegate& OnVoiceChatChannelExited() override;
+	virtual FOnVoiceChatCallStatsUpdatedDelegate& OnVoiceChatCallStatsUpdated() override;
+	virtual void Set3DPosition(const FString& ChannelName, const FVector& SpeakerPosition, const FVector& ListenerPosition, const FVector& ListenerForwardDirection, const FVector& ListenerUpDirection) override;
+	virtual TArray<FString> GetChannels() const override;
+	virtual TArray<FString> GetPlayersInChannel(const FString& ChannelName) const override;
+	virtual EVoiceChatChannelType GetChannelType(const FString& ChannelName) const override;
+	virtual FOnVoiceChatPlayerAddedDelegate& OnVoiceChatPlayerAdded() override;
+	virtual FOnVoiceChatPlayerRemovedDelegate& OnVoiceChatPlayerRemoved() override;
+	virtual bool IsPlayerTalking(const FString& PlayerName) const override;
+	virtual FOnVoiceChatPlayerTalkingUpdatedDelegate& OnVoiceChatPlayerTalkingUpdated() override;
+	virtual void SetPlayerMuted(const FString& PlayerName, bool bMuted) override;
+	virtual bool IsPlayerMuted(const FString& PlayerName) const override;
+	virtual FOnVoiceChatPlayerMuteUpdatedDelegate& OnVoiceChatPlayerMuteUpdated() override;
+	virtual void SetPlayerVolume(const FString& PlayerName, float Volume) override;
+	virtual float GetPlayerVolume(const FString& PlayerName) const override;
+	virtual FOnVoiceChatPlayerVolumeUpdatedDelegate& OnVoiceChatPlayerVolumeUpdated() override;
+	virtual void TransmitToAllChannels() override;
+	virtual void TransmitToNoChannels() override;
+	virtual void TransmitToSpecificChannel(const FString& ChannelName) override;
+	virtual EVoiceChatTransmitMode GetTransmitMode() const override;
+	virtual FString GetTransmitChannel() const override;
+	virtual FDelegateHandle StartRecording(const FOnVoiceChatRecordSamplesAvailableDelegate::FDelegate& Delegate) override;
+	virtual void StopRecording(FDelegateHandle Handle) override;
+	virtual FDelegateHandle RegisterOnVoiceChatAfterCaptureAudioReadDelegate(const FOnVoiceChatAfterCaptureAudioReadDelegate::FDelegate& Delegate) override;
+	virtual void UnregisterOnVoiceChatAfterCaptureAudioReadDelegate(FDelegateHandle Handle) override;
+	virtual FDelegateHandle RegisterOnVoiceChatBeforeCaptureAudioSentDelegate(const FOnVoiceChatBeforeCaptureAudioSentDelegate::FDelegate& Delegate) override;
+	virtual void UnregisterOnVoiceChatBeforeCaptureAudioSentDelegate(FDelegateHandle Handle) override;
+	virtual FDelegateHandle RegisterOnVoiceChatBeforeRecvAudioRenderedDelegate(const FOnVoiceChatBeforeRecvAudioRenderedDelegate::FDelegate& Delegate) override;
+	virtual void UnregisterOnVoiceChatBeforeRecvAudioRenderedDelegate(FDelegateHandle Handle) override;
+	virtual FString InsecureGetLoginToken(const FString& PlayerName) override;
+	virtual FString InsecureGetJoinToken(const FString& ChannelName, EVoiceChatChannelType ChannelType, TOptional<FVoiceChatChannel3dProperties> Channel3dProperties = TOptional<FVoiceChatChannel3dProperties>()) override;
+	// ~End IVoiceChat Interface
+
+protected:
+	template <typename TFn, typename... TArgs>
+	void DispatchUsingAccountName(const TFn& Fn, const VivoxClientApi::AccountName& AccountName, TArgs&&... Args);
+
+	template <typename TFn, typename... TArgs>
+	void DispatchAll(const TFn& Fn, TArgs&&... Args);
+
+	// ~Begin DebugClientApiEventHandler Interface
+	virtual void InvokeOnUIThread(void (Func)(void* Arg0), void* Arg0) override;
+	virtual void onLogStatementEmitted(LogLevel Level, long long NativeMillisecondsSinceEpoch, long ThreadId, const char* LogMessage) override;
+	virtual void onConnectCompleted(const VivoxClientApi::Uri& Server) override;
+	virtual void onConnectFailed(const VivoxClientApi::Uri& Server, const VivoxClientApi::VCSStatus& Status) override;
+	virtual void onDisconnected(const VivoxClientApi::Uri& Server, const VivoxClientApi::VCSStatus& Status) override;
+	virtual void onLoginCompleted(const VivoxClientApi::AccountName& AccountName) override;
+	virtual void onInvalidLoginCredentials(const VivoxClientApi::AccountName& AccountName) override;
+	virtual void onLoginFailed(const VivoxClientApi::AccountName& AccountName, const VivoxClientApi::VCSStatus& Status) override;
+	virtual void onSessionGroupCreated(const VivoxClientApi::AccountName& AccountName, const char* SessionGroupHandle) override;
+	virtual void onLogoutCompleted(const VivoxClientApi::AccountName& AccountName) override;
+	virtual void onLogoutFailed(const VivoxClientApi::AccountName& AccountName, const VivoxClientApi::VCSStatus& Status) override;
+	virtual void onChannelJoined(const VivoxClientApi::AccountName& AccountName, const VivoxClientApi::Uri& ChannelUri) override;
+	virtual void onInvalidChannelCredentials(const VivoxClientApi::AccountName& AccountName, const VivoxClientApi::Uri& ChannelUri) override;
+	virtual void onChannelJoinFailed(const VivoxClientApi::AccountName& AccountName, const VivoxClientApi::Uri& ChannelUri, const VivoxClientApi::VCSStatus& Status) override;
+	virtual void onChannelExited(const VivoxClientApi::AccountName& AccountName, const VivoxClientApi::Uri& ChannelUri, const VivoxClientApi::VCSStatus& ReasonCode) override;
+	virtual void onCallStatsUpdated(const VivoxClientApi::AccountName& AccountName, vx_call_stats_t& Stats, bool bIsFinal) override;
+	virtual void onParticipantAdded(const VivoxClientApi::AccountName& AccountName, const VivoxClientApi::Uri& ChannelUri, const VivoxClientApi::Uri& ParticipantUri, bool bIsLoggedInUser) override;
+	virtual void onParticipantLeft(const VivoxClientApi::AccountName& AccountName, const VivoxClientApi::Uri& ChannelUri, const VivoxClientApi::Uri& ParticipantUri, bool bIsLoggedInUser, ParticipantLeftReason Reason) override;
+	virtual void onParticipantUpdated(const VivoxClientApi::AccountName& AccountName, const VivoxClientApi::Uri& ChannelUri, const VivoxClientApi::Uri& ParticipantUri, bool bIsLoggedInUser, bool bSpeaking, double MeterEnergy, bool bMutedForAll) override;
+	virtual void onAvailableAudioDevicesChanged() override;
+	virtual void onOperatingSystemChosenAudioInputDeviceChanged(const VivoxClientApi::AudioDeviceId& DeviceId) override;
+	virtual void onSetApplicationChosenAudioInputDeviceCompleted(const VivoxClientApi::AccountName& AccountName, const VivoxClientApi::AudioDeviceId& DeviceId) override;
+	virtual void onSetApplicationChosenAudioInputDeviceFailed(const VivoxClientApi::AccountName& AccountName, const VivoxClientApi::AudioDeviceId& DeviceId, const VivoxClientApi::VCSStatus& Status) override;
+	virtual void onOperatingSystemChosenAudioOutputDeviceChanged(const VivoxClientApi::AudioDeviceId& DeviceId) override;
+	virtual void onSetApplicationChosenAudioOutputDeviceCompleted(const VivoxClientApi::AccountName& AccountName, const VivoxClientApi::AudioDeviceId& DeviceId) override;
+	virtual void onSetApplicationChosenAudioOutputDeviceFailed(const VivoxClientApi::AccountName& AccountName, const VivoxClientApi::AudioDeviceId& DeviceId, const VivoxClientApi::VCSStatus& Status) override;
+	virtual void onSetParticipantAudioOutputDeviceVolumeForMeCompleted(const VivoxClientApi::AccountName& AccountName, const VivoxClientApi::Uri& TargetUser, const VivoxClientApi::Uri& ChannelUri, int Volume) override;
+	virtual void onSetParticipantAudioOutputDeviceVolumeForMeFailed(const VivoxClientApi::AccountName& AccountName, const VivoxClientApi::Uri& TargetUser, const VivoxClientApi::Uri& ChannelUri, int Volume, const VivoxClientApi::VCSStatus& Status) override;
+	virtual void onSetChannelAudioOutputDeviceVolumeCompleted(const VivoxClientApi::AccountName& AccountName, const VivoxClientApi::Uri& ChannelUri, int Volume) override;
+	virtual void onSetChannelAudioOutputDeviceVolumeFailed(const VivoxClientApi::AccountName& AccountName, const VivoxClientApi::Uri& ChannelUri, int Volume, const VivoxClientApi::VCSStatus& Status) override;
+	virtual void onSetParticipantMutedForMeCompleted(const VivoxClientApi::AccountName& AccountName, const VivoxClientApi::Uri& Target, const VivoxClientApi::Uri& ChannelUri, bool bMuted) override;
+	virtual void onSetParticipantMutedForMeFailed(const VivoxClientApi::AccountName& AccountName, const VivoxClientApi::Uri& Target, const VivoxClientApi::Uri& ChannelUri, bool bMuted, const VivoxClientApi::VCSStatus& Status) override;
+	virtual void onSetChannelTransmissionToSpecificChannelCompleted(const VivoxClientApi::AccountName& AccountName, const VivoxClientApi::Uri& ChannelUri) override;
+	virtual void onSetChannelTransmissionToSpecificChannelFailed(const VivoxClientApi::AccountName& AccountName, const VivoxClientApi::Uri& ChannelUri, const VivoxClientApi::VCSStatus& Status) override;
+	virtual void onSetChannelTransmissionToAllCompleted(const VivoxClientApi::AccountName& AccountName) override;
+	virtual void onSetChannelTransmissionToAllFailed(const VivoxClientApi::AccountName& AccountName, const VivoxClientApi::VCSStatus& Status) override;
+	virtual void onSetChannelTransmissionToNoneCompleted(const VivoxClientApi::AccountName& AccountName) override;
+	virtual void onSetChannelTransmissionToNoneFailed(const VivoxClientApi::AccountName& AccountName, const VivoxClientApi::VCSStatus& Status) override;
+	virtual void onAudioUnitStarted(const char* SessionGroupHandle, const VivoxClientApi::Uri& InitialTargetUri) override;
+	virtual void onAudioUnitStopped(const char* SessionGroupHandle, const VivoxClientApi::Uri& InitialTargetUri) override;
+	virtual void onAudioUnitAfterCaptureAudioRead(const char* SessionGroupHandle, const VivoxClientApi::Uri& InitialTargetUri, short* PcmFrames, int PcmFrameCount, int AudioFrameRate, int ChannelsPerFrame) override;
+	virtual void onAudioUnitBeforeCaptureAudioSent(const char* SessionGroupHandle, const VivoxClientApi::Uri& InitialTargetUri, short* PcmFrames, int PcmFrameCount, int AudioFrameRate, int ChannelsPerFrame, bool bSpeaking) override;
+	virtual void onAudioUnitBeforeRecvAudioRendered(const char* SessionGroupHandle, const VivoxClientApi::Uri& InitialTargetUri, short* PcmFrames, int PcmFrameCount, int AudioFrameRate, int ChannelsPerFrame, bool bSilence) override;
+	// ~End DebugClientApiEventHandler Interface
+
+	virtual void SetVivoxSdkConfigHints(vx_sdk_config_t& Hints);
+
+	enum class EConnectionState
+	{
+		Disconnected,
+		Disconnecting,
+		Connecting,
+		Connected
+	};
+
+	VivoxClientApi::ClientConnection VivoxClientConnection;
+
+	bool bInitialized;
+	EConnectionState ConnectionState;
+
+
+	// Settings
+	FString VivoxServerUrl;
+	FString VivoxDomain;
+	FString VivoxIssuer;
+	FString VivoxNamespace;
+	FString VivoxInsecureSecret;
+	EVoiceChatAttenuationModel AttenuationModel;
+	int MinDistance;
+	int MaxDistance;
+	float Rolloff;
+
+	VivoxClientApi::AccountName CreateAccountName(const FString& PlayerName);
+	FString GetPlayerNameFromAccountName(const VivoxClientApi::AccountName& AccountName);
+
+	VivoxClientApi::Uri CreateUserUri(const FString& PlayerName);
+	FString GetPlayerNameFromUri(const VivoxClientApi::Uri& UserUri);
+
+	VivoxClientApi::Uri CreateChannelUri(const FString& ChannelName, EVoiceChatChannelType ChannelType, TOptional<FVoiceChatChannel3dProperties> Channel3dProperties = TOptional<FVoiceChatChannel3dProperties>());
+	FString GetChannelNameFromUri(const VivoxClientApi::Uri& ChannelUri);
+	EVoiceChatChannelType GetChannelTypeFromUri(const VivoxClientApi::Uri& ChannelUri);
+
 	// Log spam avoidance
 	FString LastLogMessage;
 	LogLevel LastLogLevel;
 	int LogSpamCount = 0;
 
-	// ~Begin FSelfRegisteringExec Interface
+	// Delegates
+	FOnVoiceChatConnectedDelegate OnVoiceChatConnectedDelegate;
+	FOnVoiceChatDisconnectedDelegate OnVoiceChatDisconnectedDelegate;
+	FOnVoiceChatReconnectedDelegate OnVoiceChatReconnectedDelegate;
+
+	// Completion delegates
+	TArray<FOnVoiceChatConnectCompleteDelegate> OnVoiceChatConnectCompleteDelegates;
+	TArray<FOnVoiceChatDisconnectCompleteDelegate> OnVoiceChatDisconnectCompleteDelegates;
+
+	// FSelfRegisteringExec Interface
 	virtual bool Exec(UWorld* InWorld, const TCHAR* Cmd, FOutputDevice& Ar) override;
-	// ~End FSelfRegisteringExec Interface
 
 	static FString ToString(EConnectionState State);
-	static FString ToString(FLoginSession::EState State);
-	static FString ToString(FChannelSession::EState State);
+
+	friend FVivoxVoiceChatUser;
+
+	FVivoxVoiceChatUser& GetVoiceChatUser();
+	FVivoxVoiceChatUser& GetVoiceChatUser() const;
+	void RegisterVoiceChatUser(FVivoxVoiceChatUser* User);
+	void UnregisterVoiceChatUser(FVivoxVoiceChatUser* User);
+	FCriticalSection VoiceChatUsersCriticalSection;
+	TArray<FVivoxVoiceChatUser*> VoiceChatUsers;
+	TUniquePtr<FVivoxVoiceChatUser> SingleUserVoiceChatUser;
 };
diff --git a/Engine/Plugins/Online/VoiceChat/VoiceChat/Source/Public/VoiceChat.h b/Engine/Plugins/Online/VoiceChat/VoiceChat/Source/Public/VoiceChat.h
--- a/Engine/Plugins/Online/VoiceChat/VoiceChat/Source/Public/VoiceChat.h
+++ b/Engine/Plugins/Online/VoiceChat/VoiceChat/Source/Public/VoiceChat.h
@@ -95,37 +95,10 @@ DECLARE_MULTICAST_DELEGATE_ThreeParams(FOnVoiceChatAfterCaptureAudioReadDelegate
 DECLARE_MULTICAST_DELEGATE_FourParams(FOnVoiceChatBeforeCaptureAudioSentDelegate, TArrayView<const int16> /* PcmSamples */, int /* SampleRate */, int /* Channels */, bool /* bIsSpeaking */);
 DECLARE_MULTICAST_DELEGATE_FourParams(FOnVoiceChatBeforeRecvAudioRenderedDelegate, TArrayView<int16> /* PcmSamples */, int /* SampleRate */, int /* Channels */, bool /* bIsSilence */);
 
-class IVoiceChat : public IModularFeature
+class IVoiceChatUser
 {
 public:
-	/**
-	 * Get the IVoiceChat modular feature instance
-	 */
-	static IVoiceChat* Get()
-	{
-		if (IModularFeatures::Get().IsModularFeatureAvailable(GetModularFeatureName()))
-		{
-			return &IModularFeatures::Get().GetModularFeature<IVoiceChat>(GetModularFeatureName());
-		}
-		return nullptr;
-	}
-
-	/**
-	 * Initialize VoiceChat
-	 */
-	virtual bool Initialize() = 0;
-
-	/**
-	 * Uninitialize VoiceChat
-	 */
-	virtual bool Uninitialize() = 0;
-
-	/**
-	 * Is voice chat initialized?
-	 *
-	 * @return true if voice chat is initialized
-	 */
-	virtual bool IsInitialized() const = 0;
+	virtual ~IVoiceChatUser() {}
 
 	/**
 	 * Set an implementation specific setting
@@ -557,6 +530,47 @@ public:
 	virtual FString InsecureGetLoginToken(const FString& PlayerName) = 0;
 	virtual FString InsecureGetJoinToken(const FString& ChannelName, EVoiceChatChannelType ChannelType, TOptional<FVoiceChatChannel3dProperties> Channel3dProperties = TOptional<FVoiceChatChannel3dProperties>()) = 0;
 
+};
+
+class IVoiceChat : public IVoiceChatUser, public IModularFeature
+{
+public:
+	/**
+	 * Get the IVoiceChat modular feature instance
+	 */
+	static IVoiceChat* Get()
+	{
+		if (IModularFeatures::Get().IsModularFeatureAvailable(GetModularFeatureName()))
+		{
+			return &IModularFeatures::Get().GetModularFeature<IVoiceChat>(GetModularFeatureName());
+		}
+		return nullptr;
+	}
+
+	/**
+	 * Initialize VoiceChat
+	 */
+	virtual bool Initialize() = 0;
+
+	/**
+	 * Uninitialize VoiceChat
+	 */
+	virtual bool Uninitialize() = 0;
+
+	/**
+	 * Is voice chat initialized?
+	 *
+	 * @return true if voice chat is initialized
+	 */
+	virtual bool IsInitialized() const = 0;
+
+	/**
+	 * Allocate an interface for an additional user // TODO: document limitations
+	 *
+	 * @return new instance of IVoiceChatUser, or nullptr if implementation does not support multiple users
+	 */
+	virtual IVoiceChatUser* CreateUser() = 0;
+
 private:
 	static FName GetModularFeatureName()
 	{
