Fix CIS

diff --git a/Engine/Plugins/Online/VoiceChat/VivoxVoiceChat/Source/Private/VivoxVoiceChat.cpp b/Engine/Plugins/Online/VoiceChat/VivoxVoiceChat/Source/Private/VivoxVoiceChat.cpp
--- a/Engine/Plugins/Online/VoiceChat/VivoxVoiceChat/Source/Private/VivoxVoiceChat.cpp
+++ b/Engine/Plugins/Online/VoiceChat/VivoxVoiceChat/Source/Private/VivoxVoiceChat.cpp
@@ -2070,7 +2070,7 @@ FString FVivoxVoiceChat::GetChannelNameFromUri(const VivoxClientApi::Uri& Channe
 		int32 Channel3dParametersIndex = Channel.Find(TEXT("!p-"));
 		if (Channel3dParametersIndex != INDEX_NONE)
 		{
-			Channel = Channel.LeftInline(Channel3dParametersIndex);
+			Channel.LeftInline(Channel3dParametersIndex);
 		}
 		return Channel;
 	}
