VivoxVoiceChat plugin - fix iOS issues introduced when support for multiple local users was added

diff --git a/Engine/Plugins/Online/VoiceChat/VivoxVoiceChat/Source/Private/IOS/IOSVivoxVoiceChat.cpp b/Engine/Plugins/Online/VoiceChat/VivoxVoiceChat/Source/Private/IOS/IOSVivoxVoiceChat.cpp
--- a/Engine/Plugins/Online/VoiceChat/VivoxVoiceChat/Source/Private/IOS/IOSVivoxVoiceChat.cpp
+++ b/Engine/Plugins/Online/VoiceChat/VivoxVoiceChat/Source/Private/IOS/IOSVivoxVoiceChat.cpp
@@ -12,12 +12,93 @@ TUniquePtr<FVivoxVoiceChat> CreateVivoxObject()
 	return MakeUnique<FIOSVivoxVoiceChat>();
 }
 
+FIOSVivoxVoiceChatUser::FIOSVivoxVoiceChatUser(FVivoxVoiceChat& InVivoxVoiceChat)
+	: FVivoxVoiceChatUser(InVivoxVoiceChat)
+	, bIsRecording(false)
+{
+}
+
+void FIOSVivoxVoiceChatUser::Initialize()
+{
+	bIsRecording = false;
+}
+
+void FIOSVivoxVoiceChatUser::SetSetting(const FString& Name, const FString& Value)
+{
+	if (Name == TEXT("HardwareAEC"))
+	{
+		GetIOSVoiceChat().SetHardwareAECEnabled(FCString::ToBool(*Value));
+	}
+	else if (Name == TEXT("BluetoothMicrophone"))
+	{
+		GetIOSVoiceChat().SetBluetoothMicrophoneEnabled(FCString::ToBool(*Value));
+	}
+	else
+	{
+		FVivoxVoiceChatUser::SetSetting(Name, Value);
+	}
+}
+
+FString FIOSVivoxVoiceChatUser::GetSetting(const FString& Name)
+{
+	if (Name == TEXT("HardwareAEC"))
+	{
+		return LexToString(GetIOSVoiceChat().IsHardwareAECEnabled());
+	}
+	else if (Name == TEXT("BluetoothMicrophone"))
+	{
+		return LexToString(GetIOSVoiceChat().IsBluetoothMicrophoneEnabled());
+	}
+	return FVivoxVoiceChatUser::GetSetting(Name);
+}
+
+void FIOSVivoxVoiceChatUser::JoinChannel(const FString& ChannelName, const FString& ChannelCredentials, EVoiceChatChannelType ChannelType, const FOnVoiceChatChannelJoinCompleteDelegate& Delegate, TOptional<FVoiceChatChannel3dProperties> Channel3dProperties)
+{
+	GetIOSVoiceChat().EnableVoiceChat(true);
+
+	FOnVoiceChatChannelJoinCompleteDelegate DelegateWrapper = FOnVoiceChatChannelJoinCompleteDelegate::CreateLambda(
+		[this, Delegate](const FString& ChannelName, const FVoiceChatResult& Result)
+			{
+				if (!Result.IsSuccess())
+				{
+					GetIOSVoiceChat().EnableVoiceChat(false);
+				}
+				Delegate.ExecuteIfBound(ChannelName, Result);
+			});
+
+	FVivoxVoiceChatUser::JoinChannel(ChannelName, ChannelCredentials, ChannelType, DelegateWrapper, Channel3dProperties);
+}
+
+FDelegateHandle FIOSVivoxVoiceChatUser::StartRecording(const FOnVoiceChatRecordSamplesAvailableDelegate::FDelegate& Delegate)
+{
+	bIsRecording = true;
+	GetIOSVoiceChat().EnableVoiceChat(true);
+	return FVivoxVoiceChatUser::StartRecording(Delegate);
+}
+
+void FIOSVivoxVoiceChatUser::StopRecording(FDelegateHandle Handle)
+{
+	FVivoxVoiceChatUser::StopRecording(Handle);
+	GetIOSVoiceChat().EnableVoiceChat(false);
+	bIsRecording = false;
+}
+
+FIOSVivoxVoiceChat& FIOSVivoxVoiceChatUser::GetIOSVoiceChat()
+{
+	return static_cast<FIOSVivoxVoiceChat&>(VivoxVoiceChat);
+}
+
+void FIOSVivoxVoiceChatUser::onChannelExited(const VivoxClientApi::AccountName& AccountName, const VivoxClientApi::Uri& ChannelUri, const VivoxClientApi::VCSStatus& Status)
+{
+	GetIOSVoiceChat().EnableVoiceChat(false);
+	FVivoxVoiceChatUser::onChannelExited(AccountName, ChannelUri, Status);
+}
+
 FIOSVivoxVoiceChat::FIOSVivoxVoiceChat()
 	: BGTask(UIBackgroundTaskInvalid)
 	, bDisconnectInBackground(true)
 	, bInBackground(false)
 	, bShouldReconnect(false)
-	, bIsRecording(false)
 	, BackgroundDelayedDisconnectTime(0.0f)
 	, DelayedDisconnectTimer(nullptr)
 {
@@ -56,7 +137,10 @@ bool FIOSVivoxVoiceChat::Initialize()
 
 	bInBackground = false;
 	bShouldReconnect = false;
-	bIsRecording = false;
+	for (IVoiceChatUser* VoiceChatUser : VoiceChatUsers)
+	{
+		static_cast<FIOSVivoxVoiceChatUser*>(VoiceChatUser)->Initialize();
+	}
 
 	return bResult;
 }
@@ -89,68 +173,63 @@ bool FIOSVivoxVoiceChat::Uninitialize()
 	return bReturn;
 }
 
-void FIOSVivoxVoiceChat::SetSetting(const FString& Name, const FString& Value)
+IVoiceChatUser* FIOSVivoxVoiceChat::CreateUser()
 {
-	if (Name == TEXT("HardwareAEC"))
-	{
-		OverrideEnableHardwareAEC = FCString::ToBool(*Value);
-		UpdateVoiceChatSettings();
-	}
-	else if (Name == TEXT("BluetoothMicrophone"))
-	{
-		OverrideEnableBluetoothMicrophone = FCString::ToBool(*Value);
-		UpdateVoiceChatSettings();
-	}
-	else
-	{
-		FVivoxVoiceChat::SetSetting(Name, Value);
-	}
+	return new FIOSVivoxVoiceChatUser(*this);
 }
 
-FString FIOSVivoxVoiceChat::GetSetting(const FString& Name)
+bool FIOSVivoxVoiceChat::IsHardwareAECEnabled() const
 {
-	if (Name == TEXT("HardwareAEC"))
-	{
-		return LexToString(IsHardwareAECEnabled());
-	}
-	else if (Name == TEXT("BluetoothMicrophone"))
-	{
-		return LexToString(IsBluetoothMicrophoneEnabled());
-	}
-	return FVivoxVoiceChat::GetSetting(Name);
+	return OverrideEnableHardwareAEC.Get(bEnableHardwareAEC);
 }
 
-void FIOSVivoxVoiceChat::JoinChannel(const FString& ChannelName, const FString& ChannelCredentials, EVoiceChatChannelType ChannelType, const FOnVoiceChatChannelJoinCompleteDelegate& Delegate, TOptional<FVoiceChatChannel3dProperties> Channel3dProperties)
+void FIOSVivoxVoiceChat::SetHardwareAECEnabled(bool bEnabled)
 {
-	EnableVoiceChat(true);
-
-	FOnVoiceChatChannelJoinCompleteDelegate DelegateWrapper = FOnVoiceChatChannelJoinCompleteDelegate::CreateLambda(
-		[this, Delegate](const FString& ChannelName, const FVoiceChatResult& Result)
-		{
-			if (!Result.IsSuccess())
-			{
-				EnableVoiceChat(false);
-			}
-			Delegate.ExecuteIfBound(ChannelName, Result);
-		});
+	OverrideEnableHardwareAEC = bEnabled;
+	UpdateVoiceChatSettings();
+}
 
-	FVivoxVoiceChat::JoinChannel(ChannelName, ChannelCredentials, ChannelType, DelegateWrapper, Channel3dProperties);
+bool FIOSVivoxVoiceChat::IsBluetoothMicrophoneEnabled() const
+{
+	return OverrideEnableBluetoothMicrophone.Get(bEnableBluetoothMicrophone);
 }
 
-FDelegateHandle FIOSVivoxVoiceChat::StartRecording(const FOnVoiceChatRecordSamplesAvailableDelegate::FDelegate& Delegate)
+void FIOSVivoxVoiceChat::SetBluetoothMicrophoneEnabled(bool bEnabled)
 {
-	bIsRecording = true;
-	EnableVoiceChat(true);
-	return FVivoxVoiceChat::StartRecording(Delegate);
+	OverrideEnableBluetoothMicrophone = bEnabled;
+	UpdateVoiceChatSettings();
 }
 
-void FIOSVivoxVoiceChat::StopRecording(FDelegateHandle Handle)
+void FIOSVivoxVoiceChat::EnableVoiceChat(bool bEnable)
 {
-	FVivoxVoiceChat::StopRecording(Handle);
-	EnableVoiceChat(false);
-	bIsRecording = false;
+	if (bEnable)
+	{
+		++VoiceChatEnableCount;
+		if (VoiceChatEnableCount == 1)
+		{
+			[[IOSAppDelegate GetDelegate] SetFeature:EAudioFeature::Playback Active:true];
+			[[IOSAppDelegate GetDelegate] SetFeature:EAudioFeature::Record Active:true];
+			
+			const bool bEnableAEC = IsHardwareAECEnabled() && IsUsingBuiltInSpeaker();
+			vx_set_platform_aec_enabled(bEnableAEC ? 1 : 0);
+		}
+	}
+	else
+	{
+		if (ensureMsgf(VoiceChatEnableCount > 0, TEXT("Attempted to disable voice chat when it was already disabled")))
+		{
+			--VoiceChatEnableCount;
+			if (VoiceChatEnableCount == 0)
+			{
+				vx_set_platform_aec_enabled(0);
+				[[IOSAppDelegate GetDelegate] SetFeature:EAudioFeature::Record Active:false];
+				[[IOSAppDelegate GetDelegate] SetFeature:EAudioFeature::Playback Active:false];
+			}
+		}
+	}
 }
 
+
 void FIOSVivoxVoiceChat::InvokeOnUIThread(void (Func)(void* Arg0), void* Arg0)
 {
 	[FIOSAsyncTask CreateTaskWithBlock:^bool()
@@ -165,20 +244,17 @@ void FIOSVivoxVoiceChat::InvokeOnUIThread(void (Func)(void* Arg0), void* Arg0)
 	FEmbeddedCommunication::WakeGameThread();
 }
 
-void FIOSVivoxVoiceChat::onChannelJoined(const VivoxClientApi::AccountName& AccountName, const VivoxClientApi::Uri& ChannelUri)
-{
-	FVivoxVoiceChat::onChannelJoined(AccountName, ChannelUri);
-}
-
-void FIOSVivoxVoiceChat::onChannelExited(const VivoxClientApi::AccountName& AccountName, const VivoxClientApi::Uri& ChannelUri, const VivoxClientApi::VCSStatus& Status)
-{
-	EnableVoiceChat(false);
-	FVivoxVoiceChat::onChannelExited(AccountName, ChannelUri, Status);
-}
-
 void FIOSVivoxVoiceChat::onDisconnected(const VivoxClientApi::Uri& Server, const VivoxClientApi::VCSStatus& Status)
 {
-	while (VoiceChatEnableCount > (bIsRecording ? 1 : 0))
+	int RecordingCount = 0;
+	for (IVoiceChatUser* VoiceChatUser : VoiceChatUsers)
+	{
+		if (static_cast<FIOSVivoxVoiceChatUser*>(VoiceChatUser)->IsRecording())
+		{
+			++RecordingCount;
+		}
+	}
+	while (VoiceChatEnableCount > RecordingCount)
 	{
 		// call once for every channel we were in
 		EnableVoiceChat(false);
@@ -313,45 +389,6 @@ void FIOSVivoxVoiceChat::Reconnect()
 	bShouldReconnect = false;
 }
 
-bool FIOSVivoxVoiceChat::IsHardwareAECEnabled() const
-{
-	return OverrideEnableHardwareAEC.Get(bEnableHardwareAEC);
-}
-
-bool FIOSVivoxVoiceChat::IsBluetoothMicrophoneEnabled() const
-{
-	return OverrideEnableBluetoothMicrophone.Get(bEnableBluetoothMicrophone);
-}
-
-void FIOSVivoxVoiceChat::EnableVoiceChat(bool bEnable)
-{
-	if (bEnable)
-	{
-		++VoiceChatEnableCount;
-		if (VoiceChatEnableCount == 1)
-		{
-			[[IOSAppDelegate GetDelegate] SetFeature:EAudioFeature::Playback Active:true];
-			[[IOSAppDelegate GetDelegate] SetFeature:EAudioFeature::Record Active:true];
-
-			const bool bEnableAEC = IsHardwareAECEnabled() && IsUsingBuiltInSpeaker();
-			vx_set_platform_aec_enabled(bEnableAEC ? 1 : 0);
-		}
-	}
-	else
-	{
-		if (ensureMsgf(VoiceChatEnableCount > 0, TEXT("Attempted to disable voice chat when it was already disabled")))
-		{
-			--VoiceChatEnableCount;
-			if (VoiceChatEnableCount == 0)
-			{
-				vx_set_platform_aec_enabled(0);
-				[[IOSAppDelegate GetDelegate] SetFeature:EAudioFeature::Record Active:false];
-				[[IOSAppDelegate GetDelegate] SetFeature:EAudioFeature::Playback Active:false];
-			}
-		}
-	}
-}
-
 void FIOSVivoxVoiceChat::UpdateVoiceChatSettings()
 {
 	if (bBluetoothMicrophoneFeatureEnabled != IsBluetoothMicrophoneEnabled())
diff --git a/Engine/Plugins/Online/VoiceChat/VivoxVoiceChat/Source/Private/IOS/IOSVivoxVoiceChat.h b/Engine/Plugins/Online/VoiceChat/VivoxVoiceChat/Source/Private/IOS/IOSVivoxVoiceChat.h
--- a/Engine/Plugins/Online/VoiceChat/VivoxVoiceChat/Source/Private/IOS/IOSVivoxVoiceChat.h
+++ b/Engine/Plugins/Online/VoiceChat/VivoxVoiceChat/Source/Private/IOS/IOSVivoxVoiceChat.h
@@ -4,6 +4,33 @@
 
 #include "VivoxVoiceChat.h"
 
+class FIOSVivoxVoiceChat;
+
+class FIOSVivoxVoiceChatUser : public FVivoxVoiceChatUser
+{
+public:
+	FIOSVivoxVoiceChatUser(FVivoxVoiceChat& InVivoxVoiceChat);
+	void Initialize();
+
+	bool IsRecording() const { return bIsRecording; }
+
+	// ~Begin IVoiceChatUser Interface
+	virtual void SetSetting(const FString& Name, const FString& Value) override;
+	virtual FString GetSetting(const FString& Name) override;
+	virtual void JoinChannel(const FString& ChannelName, const FString& ChannelCredentials, EVoiceChatChannelType ChannelType, const FOnVoiceChatChannelJoinCompleteDelegate& Delegate, TOptional<FVoiceChatChannel3dProperties> Channel3dProperties = TOptional<FVoiceChatChannel3dProperties>()) override;
+	virtual FDelegateHandle StartRecording(const FOnVoiceChatRecordSamplesAvailableDelegate::FDelegate& Delegate) override;
+	virtual void StopRecording(FDelegateHandle Handle) override;
+	// ~End IVoiceChat Interface
+	
+	FIOSVivoxVoiceChat& GetIOSVoiceChat();
+	
+private:
+	virtual void onChannelExited(const VivoxClientApi::AccountName& AccountName, const VivoxClientApi::Uri& ChannelUri, const VivoxClientApi::VCSStatus& Status) override;
+	
+	bool bIsRecording;
+};
+
+
 class FIOSVivoxVoiceChat : public FVivoxVoiceChat
 {
 public:
@@ -13,21 +40,21 @@ public:
 	// ~Begin IVoiceChat Interface
 	virtual bool Initialize() override;
 	virtual bool Uninitialize() override;
-	virtual void SetSetting(const FString& Name, const FString& Value) override;
-	virtual FString GetSetting(const FString& Name) override;
-	virtual void JoinChannel(const FString& ChannelName, const FString& ChannelCredentials, EVoiceChatChannelType ChannelType, const FOnVoiceChatChannelJoinCompleteDelegate& Delegate, TOptional<FVoiceChatChannel3dProperties> Channel3dProperties = TOptional<FVoiceChatChannel3dProperties>()) override;
-	virtual FDelegateHandle StartRecording(const FOnVoiceChatRecordSamplesAvailableDelegate::FDelegate& Delegate) override;
-	virtual void StopRecording(FDelegateHandle Handle) override;
+	virtual IVoiceChatUser* CreateUser() override;
 	// ~End IVoiceChat Interface
 
+	bool IsHardwareAECEnabled() const;
+	void SetHardwareAECEnabled(bool bEnabled);
+	bool IsBluetoothMicrophoneEnabled() const;
+	void SetBluetoothMicrophoneEnabled(bool bEnabled);
+	void EnableVoiceChat(bool bEnable);
+
 protected:
 	// ~Begin DebugClientApiEventHandler Interface
 	virtual void InvokeOnUIThread(void (Func)(void* Arg0), void* Arg0) override;
-	virtual void onChannelJoined(const VivoxClientApi::AccountName& AccountName, const VivoxClientApi::Uri& ChannelUri) override;
-	virtual void onChannelExited(const VivoxClientApi::AccountName& AccountName, const VivoxClientApi::Uri& ChannelUri, const VivoxClientApi::VCSStatus& Status) override;
 	virtual void onDisconnected(const VivoxClientApi::Uri& Server, const VivoxClientApi::VCSStatus& Status) override;
 	// ~End DebugClientApiEventHandler Interface
-	
+
 	void SetVivoxSdkConfigHints(vx_sdk_config_t &Hints) override;
 
 private:
@@ -50,24 +77,19 @@ private:
 	bool bInBackground;
 	bool bShouldReconnect;
 
-	bool bIsRecording;
-
 	float BackgroundDelayedDisconnectTime;
 	NSTimer* DelayedDisconnectTimer;
 
 	uint VoiceChatEnableCount = 0;
 
-	bool IsHardwareAECEnabled() const;
 	bool bEnableHardwareAEC = false;
 	TOptional<bool> OverrideEnableHardwareAEC;
 	bool bVoiceChatModeEnabled = false;
 
-	bool IsBluetoothMicrophoneEnabled() const;
 	bool bEnableBluetoothMicrophone = false;
 	TOptional<bool> OverrideEnableBluetoothMicrophone;
 	bool bBluetoothMicrophoneFeatureEnabled = false;
 
-	void EnableVoiceChat(bool bEnable);
 	void UpdateVoiceChatSettings();
 	bool IsUsingBuiltInSpeaker();
 };
diff --git a/Engine/Plugins/Online/VoiceChat/VivoxVoiceChat/Source/Private/VivoxVoiceChat.cpp b/Engine/Plugins/Online/VoiceChat/VivoxVoiceChat/Source/Private/VivoxVoiceChat.cpp
--- a/Engine/Plugins/Online/VoiceChat/VivoxVoiceChat/Source/Private/VivoxVoiceChat.cpp
+++ b/Engine/Plugins/Online/VoiceChat/VivoxVoiceChat/Source/Private/VivoxVoiceChat.cpp
@@ -2328,7 +2328,7 @@ FVivoxVoiceChatUser& FVivoxVoiceChat::GetVoiceChatUser()
 {
 	if (!SingleUserVoiceChatUser)
 	{
-		SingleUserVoiceChatUser = MakeUnique<FVivoxVoiceChatUser>(*this);
+		SingleUserVoiceChatUser.Reset(static_cast<FVivoxVoiceChatUser*>(CreateUser()));
 		ensureMsgf(VoiceChatUsers.Num() == 1, TEXT("When using multiple users, all connections should be managed by an IVoiceChatUser"));
 	}
 
diff --git a/Engine/Plugins/Online/VoiceChat/VivoxVoiceChat/Source/Public/VivoxVoiceChat.h b/Engine/Plugins/Online/VoiceChat/VivoxVoiceChat/Source/Public/VivoxVoiceChat.h
--- a/Engine/Plugins/Online/VoiceChat/VivoxVoiceChat/Source/Public/VivoxVoiceChat.h
+++ b/Engine/Plugins/Online/VoiceChat/VivoxVoiceChat/Source/Public/VivoxVoiceChat.h
@@ -129,7 +129,7 @@ protected:
 	void onChannelJoined(const VivoxClientApi::AccountName& AccountName, const VivoxClientApi::Uri& ChannelUri);
 	void onInvalidChannelCredentials(const VivoxClientApi::AccountName& AccountName, const VivoxClientApi::Uri& ChannelUri);
 	void onChannelJoinFailed(const VivoxClientApi::AccountName& AccountName, const VivoxClientApi::Uri& ChannelUri, const VivoxClientApi::VCSStatus& Status);
-	void onChannelExited(const VivoxClientApi::AccountName& AccountName, const VivoxClientApi::Uri& ChannelUri, const VivoxClientApi::VCSStatus& ReasonCode);
+	virtual void onChannelExited(const VivoxClientApi::AccountName& AccountName, const VivoxClientApi::Uri& ChannelUri, const VivoxClientApi::VCSStatus& ReasonCode);
 	void onCallStatsUpdated(const VivoxClientApi::AccountName& AccountName, vx_call_stats_t& Stats, bool bIsFinal);
 	void onParticipantAdded(const VivoxClientApi::AccountName& AccountName, const VivoxClientApi::Uri& ChannelUri, const VivoxClientApi::Uri& ParticipantUri, bool bIsLoggedInUser);
 	void onParticipantLeft(const VivoxClientApi::AccountName& AccountName, const VivoxClientApi::Uri& ChannelUri, const VivoxClientApi::Uri& ParticipantUri, bool bIsLoggedInUser, VivoxClientApi::IClientApiEventHandler::ParticipantLeftReason Reason);
