VivoxVoiceChatPlugin - Select the best audio device using the platfromId on platforms where this makes sense to do so

diff --git a/Engine/Plugins/Online/VoiceChat/VivoxVoiceChat/Source/Private/VivoxVoiceChat.cpp b/Engine/Plugins/Online/VoiceChat/VivoxVoiceChat/Source/Private/VivoxVoiceChat.cpp
--- a/Engine/Plugins/Online/VoiceChat/VivoxVoiceChat/Source/Private/VivoxVoiceChat.cpp
+++ b/Engine/Plugins/Online/VoiceChat/VivoxVoiceChat/Source/Private/VivoxVoiceChat.cpp
@@ -2794,7 +2794,7 @@ void FVivoxVoiceChat::DispatchUsingAccountName(const TFn& Fn, const VivoxClientA
 {
 	for (FVivoxVoiceChatUser* User : VoiceChatUsers)
 	{
-		if (User->LoginSession.AccountName == AccountName)
+		if (User->LoginSession.AccountName.IsValid() && User->LoginSession.AccountName == AccountName)
 		{
 			(User->*Fn)(AccountName, Forward<TArgs>(Args)...);
 		}
