- Make ~FVivoxVoiceChatUser() virtual

diff --git a/Engine/Plugins/Online/VoiceChat/VivoxVoiceChat/Source/Public/VivoxVoiceChat.h b/Engine/Plugins/Online/VoiceChat/VivoxVoiceChat/Source/Public/VivoxVoiceChat.h
--- a/Engine/Plugins/Online/VoiceChat/VivoxVoiceChat/Source/Public/VivoxVoiceChat.h
+++ b/Engine/Plugins/Online/VoiceChat/VivoxVoiceChat/Source/Public/VivoxVoiceChat.h
@@ -37,7 +37,7 @@ class FVivoxVoiceChatUser : public IVoiceChatUser
 {
 public:
 	FVivoxVoiceChatUser(FVivoxVoiceChat& InVivoxVoiceChat);
-	~FVivoxVoiceChatUser();
+	virtual ~FVivoxVoiceChatUser();
 
 	// ~Begin IVoiceChatUser Interface
 	virtual void SetSetting(const FString& Name, const FString& Value) override;
