fix compile error if NO_LOGGING=1

diff --git a/Engine/Plugins/Online/VoiceChat/VivoxVoiceChat/Source/Private/VivoxVoiceChat.cpp b/Engine/Plugins/Online/VoiceChat/VivoxVoiceChat/Source/Private/VivoxVoiceChat.cpp
--- a/Engine/Plugins/Online/VoiceChat/VivoxVoiceChat/Source/Private/VivoxVoiceChat.cpp
+++ b/Engine/Plugins/Online/VoiceChat/VivoxVoiceChat/Source/Private/VivoxVoiceChat.cpp
@@ -991,7 +991,11 @@ bool FVivoxVoiceChatUser::IsInitialized()
 
 bool FVivoxVoiceChatUser::Exec(UWorld* InWorld, const TCHAR* Cmd, FOutputDevice& Ar)
 {
+#if NO_LOGGING
+#define VIVOX_EXEC_LOG(...)
+#else
 #define VIVOX_EXEC_LOG(Fmt, ...) Ar.CategorizedLogf(LogVivoxVoiceChat.GetCategoryName(), ELogVerbosity::Log, Fmt, ##__VA_ARGS__)
+#endif
 
 	if (FParse::Command(&Cmd, TEXT("INFO")))
 	{
@@ -2250,7 +2254,11 @@ void FVivoxVoiceChatUser::ApplyAudioOutputDevicePolicy()
 
 bool FVivoxVoiceChat::Exec(UWorld* InWorld, const TCHAR* Cmd, FOutputDevice& Ar)
 {
+#if NO_LOGGING
+#define VIVOX_EXEC_LOG(...)
+#else
 #define VIVOX_EXEC_LOG(Fmt, ...) Ar.CategorizedLogf(LogVivoxVoiceChat.GetCategoryName(), ELogVerbosity::Log, Fmt, ##__VA_ARGS__)
+#endif
 
 	if (FParse::Command(&Cmd, TEXT("VIVOX")))
 	{
