- Fix logouts not clearing the LoginSession on VivoxVoiceChatUser
- Improve logging by printing the this pointer, so we can differentiate different users.

diff --git a/Engine/Plugins/Online/VoiceChat/VivoxVoiceChat/Source/Private/VivoxVoiceChat.cpp b/Engine/Plugins/Online/VoiceChat/VivoxVoiceChat/Source/Private/VivoxVoiceChat.cpp
--- a/Engine/Plugins/Online/VoiceChat/VivoxVoiceChat/Source/Private/VivoxVoiceChat.cpp
+++ b/Engine/Plugins/Online/VoiceChat/VivoxVoiceChat/Source/Private/VivoxVoiceChat.cpp
@@ -173,7 +173,7 @@ FString FVivoxVoiceChatUser::GetSetting(const FString& Name)
 
 void FVivoxVoiceChatUser::SetAudioInputVolume(float InVolume)
 {
-	UE_LOG(LogVivoxVoiceChat, Verbose, TEXT("SetAudioInputVolume %f"), InVolume);
+	VIVOXVOICECHATUSER_LOG(Verbose, TEXT("SetAudioInputVolume %f"), InVolume);
 
 	AudioInputOptions.Volume = InVolume;
 	ApplyAudioInputOptions();
@@ -181,7 +181,7 @@ void FVivoxVoiceChatUser::SetAudioInputVolume(float InVolume)
 
 void FVivoxVoiceChatUser::SetAudioOutputVolume(float InVolume)
 {
-	UE_LOG(LogVivoxVoiceChat, Verbose, TEXT("SetAudioOutputVolume %f"), InVolume);
+	VIVOXVOICECHATUSER_LOG(Verbose, TEXT("SetAudioOutputVolume %f"), InVolume);
 
 	AudioOutputOptions.Volume = InVolume;
 	ApplyAudioOutputOptions();
@@ -199,7 +199,7 @@ float FVivoxVoiceChatUser::GetAudioOutputVolume() const
 
 void FVivoxVoiceChatUser::SetAudioInputDeviceMuted(bool bIsMuted)
 {
-	UE_LOG(LogVivoxVoiceChat, Verbose, TEXT("SetAudioInputDeviceMuted %s"), *LexToString(bIsMuted));
+	VIVOXVOICECHATUSER_LOG(Verbose, TEXT("SetAudioInputDeviceMuted %s"), *LexToString(bIsMuted));
 
 	AudioInputOptions.bMuted = bIsMuted;
 	ApplyAudioInputOptions();
@@ -207,7 +207,7 @@ void FVivoxVoiceChatUser::SetAudioInputDeviceMuted(bool bIsMuted)
 
 void FVivoxVoiceChatUser::SetAudioOutputDeviceMuted(bool bIsMuted)
 {
-	UE_LOG(LogVivoxVoiceChat, Verbose, TEXT("SetAudioOutputDeviceMuted %s"), *LexToString(bIsMuted));
+	VIVOXVOICECHATUSER_LOG(Verbose, TEXT("SetAudioOutputDeviceMuted %s"), *LexToString(bIsMuted));
 
 	AudioOutputOptions.bMuted = bIsMuted;
 	ApplyAudioOutputOptions();
@@ -259,7 +259,7 @@ TArray<FString> FVivoxVoiceChatUser::GetAvailableOutputDevices() const
 
 void FVivoxVoiceChatUser::SetInputDevice(const FString& InputDevice)
 {
-	UE_LOG(LogVivoxVoiceChat, Verbose, TEXT("SetInputDevice %s"), *InputDevice);
+	VIVOXVOICECHATUSER_LOG(Verbose, TEXT("SetInputDevice %s"), *InputDevice);
 
 	bool bDeviceFound = false;
 	if (!InputDevice.IsEmpty())
@@ -293,7 +293,7 @@ void FVivoxVoiceChatUser::SetInputDevice(const FString& InputDevice)
 
 void FVivoxVoiceChatUser::SetOutputDevice(const FString& OutputDevice)
 {
-	UE_LOG(LogVivoxVoiceChat, Verbose, TEXT("SetOutputDevice %s"), *OutputDevice);
+	VIVOXVOICECHATUSER_LOG(Verbose, TEXT("SetOutputDevice %s"), *OutputDevice);
 
 	bool bDeviceFound = false;
 	if (!OutputDevice.IsEmpty())
@@ -435,7 +435,7 @@ void FVivoxVoiceChatUser::Login(FPlatformUserId PlatformId, const FString& Playe
 
 	if (!Result.IsSuccess())
 	{
-		UE_LOG(LogVivoxVoiceChat, Warning, TEXT("Login PlayerName:%s %s"), *PlayerName, *LexToString(Result));
+		VIVOXVOICECHATUSER_LOG(Warning, TEXT("Login PlayerName:%s %s"), *PlayerName, *LexToString(Result));
 		Delegate.ExecuteIfBound(PlayerName, Result);
 		return;
 	}
@@ -446,7 +446,7 @@ void FVivoxVoiceChatUser::Login(FPlatformUserId PlatformId, const FString& Playe
 	if (!AccountName.IsValid() || !UserUri.IsValid())
 	{
 		Result = VoiceChat::Errors::CredentialsInvalid();
-		UE_LOG(LogVivoxVoiceChat, Warning, TEXT("Login PlayerName:%s %s"), *PlayerName, *LexToString(Result));
+		VIVOXVOICECHATUSER_LOG(Warning, TEXT("Login PlayerName:%s %s"), *PlayerName, *LexToString(Result));
 		Delegate.ExecuteIfBound(PlayerName, Result);
 		return;
 	}
@@ -461,7 +461,9 @@ void FVivoxVoiceChatUser::Login(FPlatformUserId PlatformId, const FString& Playe
 	if (Status.IsError())
 	{
 		Result = ResultFromVivoxStatus(Status);
-		UE_LOG(LogVivoxVoiceChat, Warning, TEXT("Login account:%s %s"), ANSI_TO_TCHAR(LoginSession.AccountName.ToString()), *LexToString(Result));
+		VIVOXVOICECHATUSER_LOG(Warning, TEXT("Login account:%s %s"), ANSI_TO_TCHAR(LoginSession.AccountName.ToString()), *LexToString(Result));
+
+		ClearLoginSession();
 		Delegate.ExecuteIfBound(PlayerName, Result);
 		return;
 	}
@@ -491,7 +493,7 @@ void FVivoxVoiceChatUser::Logout(const FOnVoiceChatLogoutCompleteDelegate& Deleg
 
 	if (!Result.IsSuccess())
 	{
-		UE_LOG(LogVivoxVoiceChat, Warning, TEXT("Logout %s"), *LexToString(Result));
+		VIVOXVOICECHATUSER_LOG(Warning, TEXT("Logout %s"), *LexToString(Result));
 		Delegate.ExecuteIfBound(FString(), Result);
 		return;
 	}
@@ -501,7 +503,7 @@ void FVivoxVoiceChatUser::Logout(const FOnVoiceChatLogoutCompleteDelegate& Deleg
 	if (Status.IsError())
 	{
 		Result = ResultFromVivoxStatus(Status);
-		UE_LOG(LogVivoxVoiceChat, Warning, TEXT("Logout %s"), *LexToString(Result));
+		VIVOXVOICECHATUSER_LOG(Warning, TEXT("Logout %s"), *LexToString(Result));
 		Delegate.ExecuteIfBound(LoginSession.PlayerName, Result);
 		return;
 	}
@@ -536,13 +538,13 @@ void FVivoxVoiceChatUser::BlockPlayers(const TArray<FString>& PlayerNames)
 	VivoxClientApi::VCSStatus Status = VivoxClientConnection.BlockUsers(LoginSession.AccountName, UsersToBlock.GetData(), UsersToBlock.Num());
 	if (Status.IsError())
 	{
-		UE_LOG(LogVivoxVoiceChat, Warning, TEXT("BlockPlayers failed: error:%s (%i)"), ANSI_TO_TCHAR(Status.ToString()), Status.GetStatusCode());
+		VIVOXVOICECHATUSER_LOG(Warning, TEXT("BlockPlayers failed: error:%s (%i)"), ANSI_TO_TCHAR(Status.ToString()), Status.GetStatusCode());
 	}
 	else
 	{
 		for (const FString& PlayerName : PlayerNames)
 		{
-			UE_LOG(LogVivoxVoiceChat, Verbose, TEXT("Player blocked: [%s]"), *PlayerName);
+			VIVOXVOICECHATUSER_LOG(Verbose, TEXT("Player blocked: [%s]"), *PlayerName);
 		}
 	}
 }
@@ -558,13 +560,13 @@ void FVivoxVoiceChatUser::UnblockPlayers(const TArray<FString>& PlayerNames)
 	VivoxClientApi::VCSStatus Status = VivoxClientConnection.UnblockUsers(LoginSession.AccountName, UsersToUnblock.GetData(), UsersToUnblock.Num());
 	if (Status.IsError())
 	{
-		UE_LOG(LogVivoxVoiceChat, Warning, TEXT("UnblockPlayers failed: error:%s (%i)"), ANSI_TO_TCHAR(Status.ToString()), Status.GetStatusCode());
+		VIVOXVOICECHATUSER_LOG(Warning, TEXT("UnblockPlayers failed: error:%s (%i)"), ANSI_TO_TCHAR(Status.ToString()), Status.GetStatusCode());
 	}
 	else
 	{
 		for (const FString& PlayerName : PlayerNames)
 		{
-			UE_LOG(LogVivoxVoiceChat, Verbose, TEXT("Player unblocked: [%s]"), *PlayerName);
+			VIVOXVOICECHATUSER_LOG(Verbose, TEXT("Player unblocked: [%s]"), *PlayerName);
 		}
 	}
 }
@@ -601,7 +603,7 @@ void FVivoxVoiceChatUser::JoinChannel(const FString& ChannelName, const FString&
 
 	if (!Result.IsSuccess())
 	{
-		UE_LOG(LogVivoxVoiceChat, Warning, TEXT("JoinChannel ChannelName:%s %s"), *ChannelName, *LexToString(Result));
+		VIVOXVOICECHATUSER_LOG(Warning, TEXT("JoinChannel ChannelName:%s %s"), *ChannelName, *LexToString(Result));
 		Delegate.ExecuteIfBound(ChannelName, Result);
 		return;
 	}
@@ -631,7 +633,7 @@ void FVivoxVoiceChatUser::JoinChannel(const FString& ChannelName, const FString&
 	if (Status.IsError())
 	{
 		Result = ResultFromVivoxStatus(Status);
-		UE_LOG(LogVivoxVoiceChat, Warning, TEXT("JoinChannel ChannelUri:%s %s"), ANSI_TO_TCHAR(ChannelSession.ChannelUri.ToString()), *LexToString(Result));
+		VIVOXVOICECHATUSER_LOG(Warning, TEXT("JoinChannel ChannelUri:%s %s"), ANSI_TO_TCHAR(ChannelSession.ChannelUri.ToString()), *LexToString(Result));
 		Delegate.ExecuteIfBound(ChannelName, Result);
 		return;
 	}
@@ -669,7 +671,7 @@ void FVivoxVoiceChatUser::LeaveChannel(const FString& ChannelName, const FOnVoic
 
 	if (!Result.IsSuccess())
 	{
-		UE_LOG(LogVivoxVoiceChat, Warning, TEXT("LeaveChannel ChannelName:%s %s"), *ChannelName, *LexToString(Result));
+		VIVOXVOICECHATUSER_LOG(Warning, TEXT("LeaveChannel ChannelName:%s %s"), *ChannelName, *LexToString(Result));
 		Delegate.ExecuteIfBound(ChannelName, Result);
 		return;
 	}
@@ -678,7 +680,7 @@ void FVivoxVoiceChatUser::LeaveChannel(const FString& ChannelName, const FOnVoic
 	if (Status.IsError())
 	{
 		Result = ResultFromVivoxStatus(Status);
-		UE_LOG(LogVivoxVoiceChat, Warning, TEXT("LeaveChannel channel:%s %s"), ANSI_TO_TCHAR(ChannelSession.ChannelUri.ToString()), *LexToString(Result));
+		VIVOXVOICECHATUSER_LOG(Warning, TEXT("LeaveChannel channel:%s %s"), ANSI_TO_TCHAR(ChannelSession.ChannelUri.ToString()), *LexToString(Result));
 		Delegate.ExecuteIfBound(ChannelName, Result);
 		return;
 	}
@@ -699,7 +701,7 @@ void FVivoxVoiceChatUser::Set3DPosition(const FString& ChannelName, const FVecto
 	VivoxClientApi::VCSStatus Status = VivoxClientConnection.Set3DPosition(LoginSession.AccountName, ChannelSession.ChannelUri, ToVivoxVector(SpeakerPosition), ToVivoxVector(ListenerPosition), ToVivoxVector(ListenerForwardDirection), ToVivoxVector(ListenerUpDirection));
 	if (Status.IsError())
 	{
-		UE_LOG(LogVivoxVoiceChat, Warning, TEXT("Set3DPosition failed: channel:%s error:%s (%i)"), ANSI_TO_TCHAR(ChannelSession.ChannelUri.ToString()), ANSI_TO_TCHAR(Status.ToString()), Status.GetStatusCode());
+		VIVOXVOICECHATUSER_LOG(Warning, TEXT("Set3DPosition failed: channel:%s error:%s (%i)"), ANSI_TO_TCHAR(ChannelSession.ChannelUri.ToString()), ANSI_TO_TCHAR(Status.ToString()), Status.GetStatusCode());
 	}
 }
 
@@ -751,7 +753,7 @@ void FVivoxVoiceChatUser::SetPlayerMuted(const FString& PlayerName, bool bMuted)
 				VivoxClientApi::VCSStatus Status = VivoxClientConnection.SetParticipantMutedForMe(LoginSession.AccountName, ChannelParticipant->UserUri, ChannelSession.ChannelUri, bShouldMute);
 				if (Status.IsError())
 				{
-					UE_LOG(LogVivoxVoiceChat, Warning, TEXT("SetParticipantMutedForMe failed: channel:%s user:%s muted:%s error:%s (%i)"), ANSI_TO_TCHAR(ChannelSession.ChannelUri.ToString()), ANSI_TO_TCHAR(ChannelParticipant->UserUri.ToString()), *LexToString(bShouldMute), ANSI_TO_TCHAR(Status.ToString()), Status.GetStatusCode());
+					VIVOXVOICECHATUSER_LOG(Warning, TEXT("SetParticipantMutedForMe failed: channel:%s user:%s muted:%s error:%s (%i)"), ANSI_TO_TCHAR(ChannelSession.ChannelUri.ToString()), ANSI_TO_TCHAR(ChannelParticipant->UserUri.ToString()), *LexToString(bShouldMute), ANSI_TO_TCHAR(Status.ToString()), Status.GetStatusCode());
 					// TODO: This will fail only when Account/participant/channel is not found -> fixup our state
 				}
 			}
@@ -781,14 +783,14 @@ void FVivoxVoiceChatUser::SetPlayerVolume(const FString& PlayerName, float Volum
 				VivoxClientApi::VCSStatus Status = VivoxClientConnection.SetParticipantMutedForMe(LoginSession.AccountName, ChannelParticipant->UserUri, ChannelSession.ChannelUri, bShouldMute);
 				if (Status.IsError())
 				{
-					UE_LOG(LogVivoxVoiceChat, Warning, TEXT("SetParticipantMutedForMe failed: channel:%s user:%s muted:%s error:%s (%i)"), ANSI_TO_TCHAR(ChannelSession.ChannelUri.ToString()), ANSI_TO_TCHAR(ChannelParticipant->UserUri.ToString()), *LexToString(bShouldMute), ANSI_TO_TCHAR(Status.ToString()), Status.GetStatusCode());
+					VIVOXVOICECHATUSER_LOG(Warning, TEXT("SetParticipantMutedForMe failed: channel:%s user:%s muted:%s error:%s (%i)"), ANSI_TO_TCHAR(ChannelSession.ChannelUri.ToString()), ANSI_TO_TCHAR(ChannelParticipant->UserUri.ToString()), *LexToString(bShouldMute), ANSI_TO_TCHAR(Status.ToString()), Status.GetStatusCode());
 					// TODO: This will fail only when Account/participant/channel is not found -> fixup our state
 				}
 
 				Status = VivoxClientConnection.SetParticipantAudioOutputDeviceVolumeForMe(LoginSession.AccountName, ChannelParticipant->UserUri, ChannelSession.ChannelUri, Participant.IntVolume);
 				if (Status.IsError())
 				{
-					UE_LOG(LogVivoxVoiceChat, Warning, TEXT("SetParticipantAudioOutputDeviceVolumeForMe failed: channel:%s user:%s volume:%i error:%s (%i)"), ANSI_TO_TCHAR(ChannelSession.ChannelUri.ToString()), ANSI_TO_TCHAR(ChannelParticipant->UserUri.ToString()), Participant.IntVolume, ANSI_TO_TCHAR(Status.ToString()), Status.GetStatusCode());
+					VIVOXVOICECHATUSER_LOG(Warning, TEXT("SetParticipantAudioOutputDeviceVolumeForMe failed: channel:%s user:%s volume:%i error:%s (%i)"), ANSI_TO_TCHAR(ChannelSession.ChannelUri.ToString()), ANSI_TO_TCHAR(ChannelParticipant->UserUri.ToString()), Participant.IntVolume, ANSI_TO_TCHAR(Status.ToString()), Status.GetStatusCode());
 					// TODO: This will fail only when Account/participant/channel is not found -> fixup our state
 				}
 			}
@@ -803,29 +805,29 @@ float FVivoxVoiceChatUser::GetPlayerVolume(const FString& PlayerName) const
 
 void FVivoxVoiceChatUser::TransmitToNoChannels()
 {
-	UE_LOG(LogVivoxVoiceChat, Log, TEXT("TransmitToNoChannels"));
+	VIVOXVOICECHATUSER_LOG(Log, TEXT("TransmitToNoChannels"));
 
 	VivoxClientApi::VCSStatus Status = VivoxClientConnection.SetTransmissionToNone(LoginSession.AccountName);
 	if (Status.IsError())
 	{
-		UE_LOG(LogVivoxVoiceChat, Warning, TEXT("SetTransmissionToNone failed: error:%s (%i)"), ANSI_TO_TCHAR(Status.ToString()), Status.GetStatusCode());
+		VIVOXVOICECHATUSER_LOG(Warning, TEXT("SetTransmissionToNone failed: error:%s (%i)"), ANSI_TO_TCHAR(Status.ToString()), Status.GetStatusCode());
 	}
 }
 
 void FVivoxVoiceChatUser::TransmitToAllChannels()
 {
-	UE_LOG(LogVivoxVoiceChat, Log, TEXT("TransmitToAllChannels"));
+	VIVOXVOICECHATUSER_LOG(Log, TEXT("TransmitToAllChannels"));
 
 	VivoxClientApi::VCSStatus Status = VivoxClientConnection.SetTransmissionToAll(LoginSession.AccountName);
 	if (Status.IsError())
 	{
-		UE_LOG(LogVivoxVoiceChat, Warning, TEXT("SetTransmissionToAll failed: error:%s (%i)"), ANSI_TO_TCHAR(Status.ToString()), Status.GetStatusCode());
+		VIVOXVOICECHATUSER_LOG(Warning, TEXT("SetTransmissionToAll failed: error:%s (%i)"), ANSI_TO_TCHAR(Status.ToString()), Status.GetStatusCode());
 	}
 }
 
 void FVivoxVoiceChatUser::TransmitToSpecificChannel(const FString& Channel)
 {
-	UE_LOG(LogVivoxVoiceChat, Log, TEXT("TransmitToSpecificChannel %s"), *Channel);
+	VIVOXVOICECHATUSER_LOG(Log, TEXT("TransmitToSpecificChannel %s"), *Channel);
 
 	FChannelSession& ChannelSession = GetChannelSession(Channel);
 	if (ChannelSession.State == FChannelSession::EState::Connected)
@@ -833,7 +835,7 @@ void FVivoxVoiceChatUser::TransmitToSpecificChannel(const FString& Channel)
 		VivoxClientApi::VCSStatus Status = VivoxClientConnection.SetTransmissionToSpecificChannel(LoginSession.AccountName, ChannelSession.ChannelUri);
 		if (Status.IsError())
 		{
-			UE_LOG(LogVivoxVoiceChat, Warning, TEXT("TransmitToSpecificChannel failed: channel:%s error:%s (%i)"), ANSI_TO_TCHAR(ChannelSession.ChannelUri.ToString()), ANSI_TO_TCHAR(Status.ToString()), Status.GetStatusCode());
+			VIVOXVOICECHATUSER_LOG(Warning, TEXT("TransmitToSpecificChannel failed: channel:%s error:%s (%i)"), ANSI_TO_TCHAR(ChannelSession.ChannelUri.ToString()), ANSI_TO_TCHAR(Status.ToString()), Status.GetStatusCode());
 		}
 	}
 }
@@ -873,7 +875,7 @@ FDelegateHandle FVivoxVoiceChatUser::StartRecording(const FOnVoiceChatRecordSamp
 		VivoxClientApi::VCSStatus Status = VivoxClientConnection.StartAudioInputDeviceTestRecord();
 		if (Status.IsError())
 		{
-			UE_LOG(LogVivoxVoiceChat, Warning, TEXT("StartRecording failed: error:%s (%i)"), ANSI_TO_TCHAR(Status.ToString()), Status.GetStatusCode());
+			VIVOXVOICECHATUSER_LOG(Warning, TEXT("StartRecording failed: error:%s (%i)"), ANSI_TO_TCHAR(Status.ToString()), Status.GetStatusCode());
 			return FDelegateHandle();
 		}
 	}
@@ -1363,7 +1365,7 @@ void FVivoxVoiceChatUser::onDisconnected(const VivoxClientApi::Uri& Server, cons
 
 void FVivoxVoiceChatUser::onLoginCompleted(const VivoxClientApi::AccountName& AccountName)
 {
-	UE_LOG(LogVivoxVoiceChat, Log, TEXT("onLoginCompleted account:%s"), ANSI_TO_TCHAR(AccountName.ToString()));
+	VIVOXVOICECHATUSER_LOG(Log, TEXT("onLoginCompleted account:%s"), ANSI_TO_TCHAR(AccountName.ToString()));
 
 	FString PlayerName = VivoxVoiceChat.GetPlayerNameFromAccountName(AccountName);
 
@@ -1381,9 +1383,9 @@ void FVivoxVoiceChatUser::onLoginCompleted(const VivoxClientApi::AccountName& Ac
 
 void FVivoxVoiceChatUser::onInvalidLoginCredentials(const VivoxClientApi::AccountName& AccountName)
 {
-	UE_LOG(LogVivoxVoiceChat, Warning, TEXT("onInvalidLoginCredentials account:%s"), ANSI_TO_TCHAR(AccountName.ToString()));
+	VIVOXVOICECHATUSER_LOG(Warning, TEXT("onInvalidLoginCredentials account:%s"), ANSI_TO_TCHAR(AccountName.ToString()));
 
-	LoginSession.State = FLoginSession::EState::LoggedOut;
+	ClearLoginSession();
 
 	FString PlayerName = VivoxVoiceChat.GetPlayerNameFromAccountName(AccountName);
 
@@ -1392,9 +1394,9 @@ void FVivoxVoiceChatUser::onInvalidLoginCredentials(const VivoxClientApi::Accoun
 
 void FVivoxVoiceChatUser::onLoginFailed(const VivoxClientApi::AccountName& AccountName, const VivoxClientApi::VCSStatus& Status)
 {
-	UE_LOG(LogVivoxVoiceChat, Warning, TEXT("onLoginFailed account:%s error:%s (%i)"), ANSI_TO_TCHAR(AccountName.ToString()), ANSI_TO_TCHAR(Status.ToString()), Status.GetStatusCode());
+	VIVOXVOICECHATUSER_LOG(Warning, TEXT("onLoginFailed account:%s error:%s (%i)"), ANSI_TO_TCHAR(AccountName.ToString()), ANSI_TO_TCHAR(Status.ToString()), Status.GetStatusCode());
 
-	LoginSession.State = FLoginSession::EState::LoggedOut;
+	ClearLoginSession();
 
 	FString PlayerName = VivoxVoiceChat.GetPlayerNameFromAccountName(AccountName);
 
@@ -1403,11 +1405,9 @@ void FVivoxVoiceChatUser::onLoginFailed(const VivoxClientApi::AccountName& Accou
 
 void FVivoxVoiceChatUser::onLogoutCompleted(const VivoxClientApi::AccountName& AccountName)
 {
-	UE_LOG(LogVivoxVoiceChat, Log, TEXT("onLogoutCompleted account:%s"), ANSI_TO_TCHAR(AccountName.ToString()));
+	VIVOXVOICECHATUSER_LOG(Log, TEXT("onLogoutCompleted account:%s"), ANSI_TO_TCHAR(AccountName.ToString()));
 
-	ClearChannelSessions();
-
-	LoginSession.State = FLoginSession::EState::LoggedOut;
+	ClearLoginSession();
 
 	FString PlayerName = VivoxVoiceChat.GetPlayerNameFromAccountName(AccountName);
 	TriggerCompletionDelegate(OnVoiceChatLogoutCompleteDelegate, PlayerName, FVoiceChatResult::CreateSuccess());
@@ -1417,7 +1417,7 @@ void FVivoxVoiceChatUser::onLogoutCompleted(const VivoxClientApi::AccountName& A
 
 void FVivoxVoiceChatUser::onLogoutFailed(const VivoxClientApi::AccountName& AccountName, const VivoxClientApi::VCSStatus& Status)
 {
-	UE_LOG(LogVivoxVoiceChat, Warning, TEXT("onLogoutFailed account:%s error:%s (%i)"), ANSI_TO_TCHAR(AccountName.ToString()), ANSI_TO_TCHAR(Status.ToString()), Status.GetStatusCode());
+	VIVOXVOICECHATUSER_LOG(Warning, TEXT("onLogoutFailed account:%s error:%s (%i)"), ANSI_TO_TCHAR(AccountName.ToString()), ANSI_TO_TCHAR(Status.ToString()), Status.GetStatusCode());
 
 	LoginSession.State = FLoginSession::EState::LoggedIn;
 
@@ -1432,7 +1432,7 @@ void FVivoxVoiceChatUser::onSessionGroupCreated(const VivoxClientApi::AccountNam
 
 void FVivoxVoiceChatUser::onChannelJoined(const VivoxClientApi::AccountName& AccountName, const VivoxClientApi::Uri& ChannelUri)
 {
-	UE_LOG(LogVivoxVoiceChat, Log, TEXT("onChannelJoined channel:%s"), ANSI_TO_TCHAR(ChannelUri.ToString()));
+	VIVOXVOICECHATUSER_LOG(Log, TEXT("onChannelJoined channel:%s"), ANSI_TO_TCHAR(ChannelUri.ToString()));
 
 	FChannelSession& ChannelSession = GetChannelSession(ChannelUri);
 	ChannelSession.State = FChannelSession::EState::Connected;
@@ -1443,7 +1443,7 @@ void FVivoxVoiceChatUser::onChannelJoined(const VivoxClientApi::AccountName& Acc
 
 void FVivoxVoiceChatUser::onInvalidChannelCredentials(const VivoxClientApi::AccountName& AccountName, const VivoxClientApi::Uri& ChannelUri)
 {
-	UE_LOG(LogVivoxVoiceChat, Warning, TEXT("onInvalidChannelCredentials channel:%s"), ANSI_TO_TCHAR(ChannelUri.ToString()));
+	VIVOXVOICECHATUSER_LOG(Warning, TEXT("onInvalidChannelCredentials channel:%s"), ANSI_TO_TCHAR(ChannelUri.ToString()));
 
 	FChannelSession& ChannelSession = GetChannelSession(ChannelUri);
 	ChannelSession.State = FChannelSession::EState::Disconnected;
@@ -1454,7 +1454,7 @@ void FVivoxVoiceChatUser::onInvalidChannelCredentials(const VivoxClientApi::Acco
 
 void FVivoxVoiceChatUser::onChannelJoinFailed(const VivoxClientApi::AccountName& AccountName, const VivoxClientApi::Uri& ChannelUri, const VivoxClientApi::VCSStatus& Status)
 {
-	UE_LOG(LogVivoxVoiceChat, Warning, TEXT("onChannelJoinFailed channel:%s error:%s (%i)"), ANSI_TO_TCHAR(ChannelUri.ToString()), ANSI_TO_TCHAR(Status.ToString()), Status.GetStatusCode());
+	VIVOXVOICECHATUSER_LOG(Warning, TEXT("onChannelJoinFailed channel:%s error:%s (%i)"), ANSI_TO_TCHAR(ChannelUri.ToString()), ANSI_TO_TCHAR(Status.ToString()), Status.GetStatusCode());
 
 	FChannelSession& ChannelSession = GetChannelSession(ChannelUri);
 	ChannelSession.State = FChannelSession::EState::Disconnected;
@@ -1468,11 +1468,11 @@ void FVivoxVoiceChatUser::onChannelExited(const VivoxClientApi::AccountName& Acc
 	const FVoiceChatResult Result = ResultFromVivoxStatus(Status);
 	if (Result.IsSuccess())
 	{
-		UE_LOG(LogVivoxVoiceChat, Log, TEXT("onChannelExited ChannelUri:%s"), ANSI_TO_TCHAR(ChannelUri.ToString()));
+		VIVOXVOICECHATUSER_LOG(Log, TEXT("onChannelExited ChannelUri:%s"), ANSI_TO_TCHAR(ChannelUri.ToString()));
 	}
 	else
 	{
-		UE_LOG(LogVivoxVoiceChat, Warning, TEXT("onChannelExited ChannelUri:%s %s"), ANSI_TO_TCHAR(ChannelUri.ToString()), *LexToString(Result));
+		VIVOXVOICECHATUSER_LOG(Warning, TEXT("onChannelExited ChannelUri:%s %s"), ANSI_TO_TCHAR(ChannelUri.ToString()), *LexToString(Result));
 	}
 
 	FChannelSession& ChannelSession = GetChannelSession(ChannelUri);
@@ -1516,7 +1516,7 @@ void FVivoxVoiceChatUser::onCallStatsUpdated(const VivoxClientApi::AccountName&
 
 void FVivoxVoiceChatUser::onParticipantAdded(const VivoxClientApi::AccountName& AccountName, const VivoxClientApi::Uri& ChannelUri, const VivoxClientApi::Uri& ParticipantUri, bool bIsLoggedInUser)
 {
-	UE_LOG(LogVivoxVoiceChat, Log, TEXT("onParticipantAdded channel:%s participant:%s"), ANSI_TO_TCHAR(ChannelUri.ToString()), ANSI_TO_TCHAR(ParticipantUri.ToString()));
+	VIVOXVOICECHATUSER_LOG(Log, TEXT("onParticipantAdded channel:%s participant:%s"), ANSI_TO_TCHAR(ChannelUri.ToString()), ANSI_TO_TCHAR(ParticipantUri.ToString()));
 
 	FChannelSession& ChannelSession = GetChannelSession(ChannelUri);
 	FString PlayerName = VivoxVoiceChat.GetPlayerNameFromUri(ParticipantUri);
@@ -1535,7 +1535,7 @@ void FVivoxVoiceChatUser::onParticipantAdded(const VivoxClientApi::AccountName&
 		VivoxClientApi::VCSStatus Status = VivoxClientConnection.SetParticipantMutedForMe(LoginSession.AccountName, ChannelParticipant.UserUri, ChannelSession.ChannelUri, bShouldMute);
 		if (Status.IsError())
 		{
-			UE_LOG(LogVivoxVoiceChat, Warning, TEXT("SetParticipantMutedForMe failed: channel:%s user:%s muted:%s error:%s (%i)"), ANSI_TO_TCHAR(ChannelSession.ChannelUri.ToString()), ANSI_TO_TCHAR(ChannelParticipant.UserUri.ToString()), *LexToString(bShouldMute), ANSI_TO_TCHAR(Status.ToString()), Status.GetStatusCode());
+			VIVOXVOICECHATUSER_LOG(Warning, TEXT("SetParticipantMutedForMe failed: channel:%s user:%s muted:%s error:%s (%i)"), ANSI_TO_TCHAR(ChannelSession.ChannelUri.ToString()), ANSI_TO_TCHAR(ChannelParticipant.UserUri.ToString()), *LexToString(bShouldMute), ANSI_TO_TCHAR(Status.ToString()), Status.GetStatusCode());
 			// TODO: This will fail only when Account/participant/channel is not found -> fixup our state
 		}
 	}
@@ -1547,14 +1547,14 @@ void FVivoxVoiceChatUser::onParticipantAdded(const VivoxClientApi::AccountName&
 		VivoxClientApi::VCSStatus Status = VivoxClientConnection.SetParticipantMutedForMe(LoginSession.AccountName, ChannelParticipant.UserUri, ChannelSession.ChannelUri, bShouldMute);
 		if (Status.IsError())
 		{
-			UE_LOG(LogVivoxVoiceChat, Warning, TEXT("SetParticipantMutedForMe failed: channel:%s user:%s muted:%s error:%s (%i)"), ANSI_TO_TCHAR(ChannelSession.ChannelUri.ToString()), ANSI_TO_TCHAR(ChannelParticipant.UserUri.ToString()), *LexToString(bShouldMute), ANSI_TO_TCHAR(Status.ToString()), Status.GetStatusCode());
+			VIVOXVOICECHATUSER_LOG(Warning, TEXT("SetParticipantMutedForMe failed: channel:%s user:%s muted:%s error:%s (%i)"), ANSI_TO_TCHAR(ChannelSession.ChannelUri.ToString()), ANSI_TO_TCHAR(ChannelParticipant.UserUri.ToString()), *LexToString(bShouldMute), ANSI_TO_TCHAR(Status.ToString()), Status.GetStatusCode());
 			// TODO: This will fail only when Account/participant/channel is not found -> fixup our state
 		}
 
 		Status = VivoxClientConnection.SetParticipantAudioOutputDeviceVolumeForMe(LoginSession.AccountName, ChannelParticipant.UserUri, ChannelSession.ChannelUri, Participant.IntVolume);
 		if (Status.IsError())
 		{
-			UE_LOG(LogVivoxVoiceChat, Warning, TEXT("SetParticipantAudioOutputDeviceVolumeForMe failed: channel:%s user:%s volume:%i error:%s (%i)"), ANSI_TO_TCHAR(ChannelSession.ChannelUri.ToString()), ANSI_TO_TCHAR(ChannelParticipant.UserUri.ToString()), Participant.IntVolume, ANSI_TO_TCHAR(Status.ToString()), Status.GetStatusCode());
+			VIVOXVOICECHATUSER_LOG(Warning, TEXT("SetParticipantAudioOutputDeviceVolumeForMe failed: channel:%s user:%s volume:%i error:%s (%i)"), ANSI_TO_TCHAR(ChannelSession.ChannelUri.ToString()), ANSI_TO_TCHAR(ChannelParticipant.UserUri.ToString()), Participant.IntVolume, ANSI_TO_TCHAR(Status.ToString()), Status.GetStatusCode());
 			// TODO: This will fail only when Account/participant/channel is not found -> fixup our state
 		}
 	}
@@ -1562,7 +1562,7 @@ void FVivoxVoiceChatUser::onParticipantAdded(const VivoxClientApi::AccountName&
 
 void FVivoxVoiceChatUser::onParticipantLeft(const VivoxClientApi::AccountName& AccountName, const VivoxClientApi::Uri& ChannelUri, const VivoxClientApi::Uri& ParticipantUri, bool bIsLoggedInUser, VivoxClientApi::IClientApiEventHandler::ParticipantLeftReason Reason)
 {
-	UE_LOG(LogVivoxVoiceChat, Log, TEXT("onParticipantLeft channel:%s participant:%s reason:%s"), ANSI_TO_TCHAR(ChannelUri.ToString()), ANSI_TO_TCHAR(ParticipantUri.ToString()), *LexToString(Reason));
+	VIVOXVOICECHATUSER_LOG(Log, TEXT("onParticipantLeft channel:%s participant:%s reason:%s"), ANSI_TO_TCHAR(ChannelUri.ToString()), ANSI_TO_TCHAR(ParticipantUri.ToString()), *LexToString(Reason));
 
 	FChannelSession& ChannelSession = GetChannelSession(ChannelUri);
 	FString PlayerName = VivoxVoiceChat.GetPlayerNameFromUri(ParticipantUri);
@@ -1574,7 +1574,7 @@ void FVivoxVoiceChatUser::onParticipantLeft(const VivoxClientApi::AccountName& A
 
 void FVivoxVoiceChatUser::onParticipantUpdated(const VivoxClientApi::AccountName& AccountName, const VivoxClientApi::Uri& ChannelUri, const VivoxClientApi::Uri& ParticipantUri, bool bIsLoggedInUser, bool bSpeaking, double MeterEnergy, bool bMutedForAll)
 {
-	UE_LOG(LogVivoxVoiceChat, VeryVerbose, TEXT("onParticipantUpdated channel:%s participant:%s speaking:%s energy:%f"), ANSI_TO_TCHAR(ChannelUri.ToString()), ANSI_TO_TCHAR(ParticipantUri.ToString()), *LexToString(bSpeaking), MeterEnergy);
+	VIVOXVOICECHATUSER_LOG(VeryVerbose, TEXT("onParticipantUpdated channel:%s participant:%s speaking:%s energy:%f"), ANSI_TO_TCHAR(ChannelUri.ToString()), ANSI_TO_TCHAR(ParticipantUri.ToString()), *LexToString(bSpeaking), MeterEnergy);
 
 	FChannelSession& ChannelSession = GetChannelSession(ChannelUri);
 	FString PlayerName = VivoxVoiceChat.GetPlayerNameFromUri(ParticipantUri);
@@ -1591,44 +1591,44 @@ void FVivoxVoiceChatUser::onParticipantUpdated(const VivoxClientApi::AccountName
 
 void FVivoxVoiceChatUser::onAvailableAudioDevicesChanged()
 {
-	UE_LOG(LogVivoxVoiceChat, Verbose, TEXT("onAvailableAudioDevicesChanged"));
+	VIVOXVOICECHATUSER_LOG(Verbose, TEXT("onAvailableAudioDevicesChanged"));
 
 	OnVoiceChatAvailableAudioDevicesChangedDelegate.Broadcast();
 }
 
 void FVivoxVoiceChatUser::onOperatingSystemChosenAudioInputDeviceChanged(const VivoxClientApi::AudioDeviceId& DeviceId)
 {
-	UE_LOG(LogVivoxVoiceChat, Verbose, TEXT("onOperatingSystemChosenAudioInputDeviceChanged deviceid:%s"), UTF8_TO_TCHAR(DeviceId.ToString()));
+	VIVOXVOICECHATUSER_LOG(Verbose, TEXT("onOperatingSystemChosenAudioInputDeviceChanged deviceid:%s"), UTF8_TO_TCHAR(DeviceId.ToString()));
 }
 
 void FVivoxVoiceChatUser::onSetApplicationChosenAudioInputDeviceCompleted(const VivoxClientApi::AccountName& AccountName, const VivoxClientApi::AudioDeviceId& DeviceId)
 {
-	UE_LOG(LogVivoxVoiceChat, Verbose, TEXT("onSetApplicationChosenAudioInputDeviceCompleted deviceid:%s"), UTF8_TO_TCHAR(DeviceId.ToString()));
+	VIVOXVOICECHATUSER_LOG(Verbose, TEXT("onSetApplicationChosenAudioInputDeviceCompleted deviceid:%s"), UTF8_TO_TCHAR(DeviceId.ToString()));
 }
 
 void FVivoxVoiceChatUser::onSetApplicationChosenAudioInputDeviceFailed(const VivoxClientApi::AccountName& AccountName, const VivoxClientApi::AudioDeviceId& DeviceId, const VivoxClientApi::VCSStatus& Status)
 {
-	UE_LOG(LogVivoxVoiceChat, Warning, TEXT("onSetApplicationChosenAudioInputDeviceFailed deviceid:%s error:%s (%i)"), UTF8_TO_TCHAR(DeviceId.ToString()), ANSI_TO_TCHAR(Status.ToString()), Status.GetStatusCode());
+	VIVOXVOICECHATUSER_LOG(Warning, TEXT("onSetApplicationChosenAudioInputDeviceFailed deviceid:%s error:%s (%i)"), UTF8_TO_TCHAR(DeviceId.ToString()), ANSI_TO_TCHAR(Status.ToString()), Status.GetStatusCode());
 }
 
 void FVivoxVoiceChatUser::onOperatingSystemChosenAudioOutputDeviceChanged(const VivoxClientApi::AudioDeviceId& DeviceId)
 {
-	UE_LOG(LogVivoxVoiceChat, Verbose, TEXT("onOperatingSystemChosenAudioOutputDeviceChanged deviceid:%s"), UTF8_TO_TCHAR(DeviceId.ToString()));
+	VIVOXVOICECHATUSER_LOG(Verbose, TEXT("onOperatingSystemChosenAudioOutputDeviceChanged deviceid:%s"), UTF8_TO_TCHAR(DeviceId.ToString()));
 }
 
 void FVivoxVoiceChatUser::onSetApplicationChosenAudioOutputDeviceCompleted(const VivoxClientApi::AccountName& AccountName, const VivoxClientApi::AudioDeviceId& DeviceId)
 {
-	UE_LOG(LogVivoxVoiceChat, Verbose, TEXT("onSetApplicationChosenAudioOutputDeviceCompleted deviceid:%s"), UTF8_TO_TCHAR(DeviceId.ToString()));
+	VIVOXVOICECHATUSER_LOG(Verbose, TEXT("onSetApplicationChosenAudioOutputDeviceCompleted deviceid:%s"), UTF8_TO_TCHAR(DeviceId.ToString()));
 }
 
 void FVivoxVoiceChatUser::onSetApplicationChosenAudioOutputDeviceFailed(const VivoxClientApi::AccountName& AccountName, const VivoxClientApi::AudioDeviceId& DeviceId, const VivoxClientApi::VCSStatus& Status)
 {
-	UE_LOG(LogVivoxVoiceChat, Warning, TEXT("onSetApplicationChosenAudioOutputDeviceFailed deviceid:%s error:%s (%i)"), UTF8_TO_TCHAR(DeviceId.ToString()), ANSI_TO_TCHAR(Status.ToString()), Status.GetStatusCode());
+	VIVOXVOICECHATUSER_LOG(Warning, TEXT("onSetApplicationChosenAudioOutputDeviceFailed deviceid:%s error:%s (%i)"), UTF8_TO_TCHAR(DeviceId.ToString()), ANSI_TO_TCHAR(Status.ToString()), Status.GetStatusCode());
 }
 
 void FVivoxVoiceChatUser::onSetParticipantAudioOutputDeviceVolumeForMeCompleted(const VivoxClientApi::AccountName& AccountName, const VivoxClientApi::Uri& TargetUser, const VivoxClientApi::Uri& ChannelUri, int Volume)
 {
-	UE_LOG(LogVivoxVoiceChat, Verbose, TEXT("onSetParticipantAudioOutputDeviceVolumeForMeCompleted channel:%s user:%s volume:%i"), ANSI_TO_TCHAR(ChannelUri.ToString()), ANSI_TO_TCHAR(TargetUser.ToString()), Volume);
+	VIVOXVOICECHATUSER_LOG(Verbose, TEXT("onSetParticipantAudioOutputDeviceVolumeForMeCompleted channel:%s user:%s volume:%i"), ANSI_TO_TCHAR(ChannelUri.ToString()), ANSI_TO_TCHAR(TargetUser.ToString()), Volume);
 
 	FChannelSession& ChannelSession = GetChannelSession(ChannelUri);
 	FString PlayerName = VivoxVoiceChat.GetPlayerNameFromUri(TargetUser);
@@ -1642,7 +1642,7 @@ void FVivoxVoiceChatUser::onSetParticipantAudioOutputDeviceVolumeForMeCompleted(
 
 void FVivoxVoiceChatUser::onSetParticipantAudioOutputDeviceVolumeForMeFailed(const VivoxClientApi::AccountName& AccountName, const VivoxClientApi::Uri& TargetUser, const VivoxClientApi::Uri& ChannelUri, int Volume, const VivoxClientApi::VCSStatus& Status)
 {
-	UE_LOG(LogVivoxVoiceChat, Warning, TEXT("onSetParticipantAudioOutputDeviceVolumeForMeFailed channel:%s user:%s volume:%i error:%s (%i)"), ANSI_TO_TCHAR(ChannelUri.ToString()), ANSI_TO_TCHAR(TargetUser.ToString()), Volume, ANSI_TO_TCHAR(Status.ToString()), Status.GetStatusCode());
+	VIVOXVOICECHATUSER_LOG(Warning, TEXT("onSetParticipantAudioOutputDeviceVolumeForMeFailed channel:%s user:%s volume:%i error:%s (%i)"), ANSI_TO_TCHAR(ChannelUri.ToString()), ANSI_TO_TCHAR(TargetUser.ToString()), Volume, ANSI_TO_TCHAR(Status.ToString()), Status.GetStatusCode());
 
 	FChannelSession& ChannelSession = GetChannelSession(ChannelUri);
 	FString PlayerName = VivoxVoiceChat.GetPlayerNameFromUri(TargetUser);
@@ -1654,17 +1654,17 @@ void FVivoxVoiceChatUser::onSetParticipantAudioOutputDeviceVolumeForMeFailed(con
 
 void FVivoxVoiceChatUser::onSetChannelAudioOutputDeviceVolumeCompleted(const VivoxClientApi::AccountName& AccountName, const VivoxClientApi::Uri& ChannelUri, int Volume)
 {
-	UE_LOG(LogVivoxVoiceChat, Verbose, TEXT("onSetChannelAudioOutputDeviceVolumeCompleted channel:%s volume:%i"), ANSI_TO_TCHAR(ChannelUri.ToString()), Volume);
+	VIVOXVOICECHATUSER_LOG(Verbose, TEXT("onSetChannelAudioOutputDeviceVolumeCompleted channel:%s volume:%i"), ANSI_TO_TCHAR(ChannelUri.ToString()), Volume);
 }
 
 void FVivoxVoiceChatUser::onSetChannelAudioOutputDeviceVolumeFailed(const VivoxClientApi::AccountName& AccountName, const VivoxClientApi::Uri& ChannelUri, int Volume, const VivoxClientApi::VCSStatus& Status)
 {
-	UE_LOG(LogVivoxVoiceChat, Warning, TEXT("onSetChannelAudioOutputDeviceVolumeFailed channel:%s volume:%i error:%s (%i)"), ANSI_TO_TCHAR(ChannelUri.ToString()), Volume, ANSI_TO_TCHAR(Status.ToString()), Status.GetStatusCode());
+	VIVOXVOICECHATUSER_LOG(Warning, TEXT("onSetChannelAudioOutputDeviceVolumeFailed channel:%s volume:%i error:%s (%i)"), ANSI_TO_TCHAR(ChannelUri.ToString()), Volume, ANSI_TO_TCHAR(Status.ToString()), Status.GetStatusCode());
 }
 
 void FVivoxVoiceChatUser::onSetParticipantMutedForMeCompleted(const VivoxClientApi::AccountName& AccountName, const VivoxClientApi::Uri& Target, const VivoxClientApi::Uri& ChannelUri, bool bMuted)
 {
-	UE_LOG(LogVivoxVoiceChat, Verbose, TEXT("onSetParticipantMutedForMeCompleted channel:%s user:%s muted:%s"), ANSI_TO_TCHAR(ChannelUri.ToString()), ANSI_TO_TCHAR(Target.ToString()), *LexToString(bMuted));
+	VIVOXVOICECHATUSER_LOG(Verbose, TEXT("onSetParticipantMutedForMeCompleted channel:%s user:%s muted:%s"), ANSI_TO_TCHAR(ChannelUri.ToString()), ANSI_TO_TCHAR(Target.ToString()), *LexToString(bMuted));
 
 	FChannelSession& ChannelSession = GetChannelSession(ChannelUri);
 	FString PlayerName = VivoxVoiceChat.GetPlayerNameFromUri(Target);
@@ -1678,7 +1678,7 @@ void FVivoxVoiceChatUser::onSetParticipantMutedForMeCompleted(const VivoxClientA
 
 void FVivoxVoiceChatUser::onSetParticipantMutedForMeFailed(const VivoxClientApi::AccountName& AccountName, const VivoxClientApi::Uri& Target, const VivoxClientApi::Uri& ChannelUri, bool bMuted, const VivoxClientApi::VCSStatus& Status)
 {
-	UE_LOG(LogVivoxVoiceChat, Warning, TEXT("onSetParticipantMutedForMeFailed channel:%s user:%s muted:%s error:%s (%i)"), ANSI_TO_TCHAR(ChannelUri.ToString()), ANSI_TO_TCHAR(Target.ToString()), *LexToString(bMuted), ANSI_TO_TCHAR(Status.ToString()), Status.GetStatusCode());
+	VIVOXVOICECHATUSER_LOG(Warning, TEXT("onSetParticipantMutedForMeFailed channel:%s user:%s muted:%s error:%s (%i)"), ANSI_TO_TCHAR(ChannelUri.ToString()), ANSI_TO_TCHAR(Target.ToString()), *LexToString(bMuted), ANSI_TO_TCHAR(Status.ToString()), Status.GetStatusCode());
 
 	FChannelSession& ChannelSession = GetChannelSession(ChannelUri);
 	FString PlayerName = VivoxVoiceChat.GetPlayerNameFromUri(Target);
@@ -1690,32 +1690,32 @@ void FVivoxVoiceChatUser::onSetParticipantMutedForMeFailed(const VivoxClientApi:
 
 void FVivoxVoiceChatUser::onSetChannelTransmissionToSpecificChannelCompleted(const VivoxClientApi::AccountName& AccountName, const VivoxClientApi::Uri& ChannelUri)
 {
-	UE_LOG(LogVivoxVoiceChat, Verbose, TEXT("onSetChannelTransmissionToSpecificChannelCompleted channel:%s"), ANSI_TO_TCHAR(ChannelUri.ToString()));
+	VIVOXVOICECHATUSER_LOG(Verbose, TEXT("onSetChannelTransmissionToSpecificChannelCompleted channel:%s"), ANSI_TO_TCHAR(ChannelUri.ToString()));
 }
 
 void FVivoxVoiceChatUser::onSetChannelTransmissionToSpecificChannelFailed(const VivoxClientApi::AccountName& AccountName, const VivoxClientApi::Uri& ChannelUri, const VivoxClientApi::VCSStatus& Status)
 {
-	UE_LOG(LogVivoxVoiceChat, Warning, TEXT("onSetChannelTransmissionToSpecificChannelCompleted channel:%s error:%s (%i)"), ANSI_TO_TCHAR(ChannelUri.ToString()), ANSI_TO_TCHAR(Status.ToString()), Status.GetStatusCode());
+	VIVOXVOICECHATUSER_LOG(Warning, TEXT("onSetChannelTransmissionToSpecificChannelCompleted channel:%s error:%s (%i)"), ANSI_TO_TCHAR(ChannelUri.ToString()), ANSI_TO_TCHAR(Status.ToString()), Status.GetStatusCode());
 }
 
 void FVivoxVoiceChatUser::onSetChannelTransmissionToAllCompleted(const VivoxClientApi::AccountName& AccountName)
 {
-	UE_LOG(LogVivoxVoiceChat, Verbose, TEXT("onSetChannelTransmissionToAllCompleted"));
+	VIVOXVOICECHATUSER_LOG(Verbose, TEXT("onSetChannelTransmissionToAllCompleted"));
 }
 
 void FVivoxVoiceChatUser::onSetChannelTransmissionToAllFailed(const VivoxClientApi::AccountName& AccountName, const VivoxClientApi::VCSStatus& Status)
 {
-	UE_LOG(LogVivoxVoiceChat, Warning, TEXT("onSetChannelTransmissionToAllFailed error:%s (%i)"), ANSI_TO_TCHAR(Status.ToString()), Status.GetStatusCode());
+	VIVOXVOICECHATUSER_LOG(Warning, TEXT("onSetChannelTransmissionToAllFailed error:%s (%i)"), ANSI_TO_TCHAR(Status.ToString()), Status.GetStatusCode());
 }
 
 void FVivoxVoiceChatUser::onSetChannelTransmissionToNoneCompleted(const VivoxClientApi::AccountName& AccountName)
 {
-	UE_LOG(LogVivoxVoiceChat, Verbose, TEXT("onSetChannelTransmissionToNoneCompleted"));
+	VIVOXVOICECHATUSER_LOG(Verbose, TEXT("onSetChannelTransmissionToNoneCompleted"));
 }
 
 void FVivoxVoiceChatUser::onSetChannelTransmissionToNoneFailed(const VivoxClientApi::AccountName& AccountName, const VivoxClientApi::VCSStatus& Status)
 {
-	UE_LOG(LogVivoxVoiceChat, Warning, TEXT("onSetChannelTransmissionToNoneFailed error:%s (%i)"), ANSI_TO_TCHAR(Status.ToString()), Status.GetStatusCode());
+	VIVOXVOICECHATUSER_LOG(Warning, TEXT("onSetChannelTransmissionToNoneFailed error:%s (%i)"), ANSI_TO_TCHAR(Status.ToString()), Status.GetStatusCode());
 }
 
 void FVivoxVoiceChatUser::onAudioUnitStarted(const char* SessionGroupHandle, const VivoxClientApi::Uri& InitialTargetUri)
@@ -2197,7 +2197,7 @@ void FVivoxVoiceChatUser::ClearChannelSessions()
 void FVivoxVoiceChatUser::ClearLoginSession()
 {
 	ClearChannelSessions();
-	LoginSession.State = FLoginSession::EState::LoggedOut;
+	LoginSession = FLoginSession();
 }
 
 void FVivoxVoiceChatUser::ApplyAudioInputOptions()
diff --git a/Engine/Plugins/Online/VoiceChat/VivoxVoiceChat/Source/Public/VivoxVoiceChat.h b/Engine/Plugins/Online/VoiceChat/VivoxVoiceChat/Source/Public/VivoxVoiceChat.h
--- a/Engine/Plugins/Online/VoiceChat/VivoxVoiceChat/Source/Public/VivoxVoiceChat.h
+++ b/Engine/Plugins/Online/VoiceChat/VivoxVoiceChat/Source/Public/VivoxVoiceChat.h
@@ -21,6 +21,11 @@ THIRD_PARTY_INCLUDES_END
 
 DECLARE_LOG_CATEGORY_EXTERN(LogVivoxVoiceChat, Log, All);
 
+#define VIVOXVOICECHATUSER_LOG(VivoxVoiceChatLogLevel, VivoxVoiceChatFormatStr, ...) \
+{ \
+	UE_LOG(LogVivoxVoiceChat, VivoxVoiceChatLogLevel, TEXT("[%p] ") VivoxVoiceChatFormatStr, (void*)this, __VA_ARGS__); \
+}
+
 DECLARE_STATS_GROUP(TEXT("Vivox"), STATGROUP_Vivox, STATCAT_Advanced);
 
 class VIVOXVOICECHAT_API FVivoxDelegates
@@ -190,7 +195,7 @@ protected:
 
 	struct FLoginSession
 	{
-		FPlatformUserId PlatformId;
+		FPlatformUserId PlatformId = PLATFORMUSERID_NONE;
 		FString PlayerName;
 		VivoxClientApi::AccountName AccountName;
 		VivoxClientApi::Uri UserUri;
