- Fix mac compilation errors due to incorrect use of __VA_ARGS__

diff --git a/Engine/Plugins/Online/VoiceChat/VivoxVoiceChat/Source/Public/VivoxVoiceChat.h b/Engine/Plugins/Online/VoiceChat/VivoxVoiceChat/Source/Public/VivoxVoiceChat.h
--- a/Engine/Plugins/Online/VoiceChat/VivoxVoiceChat/Source/Public/VivoxVoiceChat.h
+++ b/Engine/Plugins/Online/VoiceChat/VivoxVoiceChat/Source/Public/VivoxVoiceChat.h
@@ -23,7 +23,7 @@ DECLARE_LOG_CATEGORY_EXTERN(LogVivoxVoiceChat, Log, All);
 
 #define VIVOXVOICECHATUSER_LOG(VivoxVoiceChatLogLevel, VivoxVoiceChatFormatStr, ...) \
 { \
-	UE_LOG(LogVivoxVoiceChat, VivoxVoiceChatLogLevel, TEXT("[%p] ") VivoxVoiceChatFormatStr, (void*)this, __VA_ARGS__); \
+	UE_LOG(LogVivoxVoiceChat, VivoxVoiceChatLogLevel, TEXT("[%p] ") VivoxVoiceChatFormatStr, (void*)this, ##__VA_ARGS__); \
 }
 
 DECLARE_STATS_GROUP(TEXT("Vivox"), STATGROUP_Vivox, STATCAT_Advanced);
