- Fix incorrect VoiceChatUser indirection

diff --git a/Engine/Plugins/Online/VoiceChat/VivoxVoiceChat/Source/Private/VivoxVoiceChat.cpp b/Engine/Plugins/Online/VoiceChat/VivoxVoiceChat/Source/Private/VivoxVoiceChat.cpp
--- a/Engine/Plugins/Online/VoiceChat/VivoxVoiceChat/Source/Private/VivoxVoiceChat.cpp
+++ b/Engine/Plugins/Online/VoiceChat/VivoxVoiceChat/Source/Private/VivoxVoiceChat.cpp
@@ -2452,7 +2452,7 @@ bool FVivoxVoiceChat::GetAudioOutputDeviceMuted() const
 
 TArray<FString> FVivoxVoiceChat::GetAvailableInputDevices() const
 {
-	return GetVoiceChatUser().GetAvailableOutputDevices();
+	return GetVoiceChatUser().GetAvailableInputDevices();
 }
 
 TArray<FString> FVivoxVoiceChat::GetAvailableOutputDevices() const
