VivoxVoiceChat plugin: Add VIVOX CREATEUSER console command to allow plugin to be tested purely with console commands

diff --git a/Engine/Plugins/Online/VoiceChat/VivoxVoiceChat/Source/Private/VivoxVoiceChat.cpp b/Engine/Plugins/Online/VoiceChat/VivoxVoiceChat/Source/Private/VivoxVoiceChat.cpp
--- a/Engine/Plugins/Online/VoiceChat/VivoxVoiceChat/Source/Private/VivoxVoiceChat.cpp
+++ b/Engine/Plugins/Online/VoiceChat/VivoxVoiceChat/Source/Private/VivoxVoiceChat.cpp
@@ -2308,6 +2308,11 @@ bool FVivoxVoiceChat::Exec(UWorld* InWorld, const TCHAR* Cmd, FOutputDevice& Ar)
 			Uninitialize();
 			return true;
 		}
+		else if (FParse::Command(&Cmd, TEXT("CREATEUSER")))
+		{
+			(void)CreateUser();
+			return true;
+		}
 		else
 		{
 			int Index = 0;
@@ -2318,7 +2323,15 @@ bool FVivoxVoiceChat::Exec(UWorld* InWorld, const TCHAR* Cmd, FOutputDevice& Ar)
 				{
 					if (FVivoxVoiceChatUser* User = VoiceChatUsers[Index])
 					{
-						return User->Exec(InWorld, Cmd, Ar);
+						if (FParse::Command(&Cmd, TEXT("DESTROYUSER")))
+						{
+							delete User;
+							return true;
+						}
+						else
+						{
+							return User->Exec(InWorld, Cmd, Ar);
+						}
 					}
 				}
 			}
