- Refactor VoiceChat interface to use device id's instead of names when selecting input/output devices

diff --git a/Engine/Plugins/Online/VoiceChat/VivoxVoiceChat/Source/Private/VivoxVoiceChat.cpp b/Engine/Plugins/Online/VoiceChat/VivoxVoiceChat/Source/Private/VivoxVoiceChat.cpp
--- a/Engine/Plugins/Online/VoiceChat/VivoxVoiceChat/Source/Private/VivoxVoiceChat.cpp
+++ b/Engine/Plugins/Online/VoiceChat/VivoxVoiceChat/Source/Private/VivoxVoiceChat.cpp
@@ -223,9 +223,9 @@ bool FVivoxVoiceChatUser::GetAudioOutputDeviceMuted() const
 	return AudioOutputOptions.bMuted;
 }
 
-TArray<FString> FVivoxVoiceChatUser::GetAvailableInputDevices() const
+TArray<FVoiceChatDeviceInfo> FVivoxVoiceChatUser::GetAvailableInputDeviceInfos() const
 {
-	TArray<FString> InputDevices;
+	TArray<FVoiceChatDeviceInfo> InputDevices;
 	const VivoxClientApi::AudioDeviceId* AudioDevices = nullptr;
 	int NumAudioDevices = 0;
 	VivoxClientConnection.GetAvailableAudioInputDevices(LoginSession.AccountName, AudioDevices, NumAudioDevices);
@@ -234,15 +234,17 @@ TArray<FString> FVivoxVoiceChatUser::GetAvailableInputDevices() const
 	{
 		for (int DeviceIndex = 0; DeviceIndex < NumAudioDevices; ++DeviceIndex)
 		{
-			InputDevices.Add(UTF8_TO_TCHAR(AudioDevices[DeviceIndex].GetAudioDeviceDisplayName()));
+			FVoiceChatDeviceInfo& DeviceInfo = InputDevices.Emplace_GetRef();
+			DeviceInfo.DisplayName = UTF8_TO_TCHAR(AudioDevices[DeviceIndex].GetAudioDeviceDisplayName());
+			DeviceInfo.Id = UTF8_TO_TCHAR(AudioDevices[DeviceIndex].GetAudioDeviceId());
 		}
 	}
 	return InputDevices;
 }
 
-TArray<FString> FVivoxVoiceChatUser::GetAvailableOutputDevices() const
+TArray<FVoiceChatDeviceInfo> FVivoxVoiceChatUser::GetAvailableOutputDeviceInfos() const
 {
-	TArray<FString> OutputDevices;
+	TArray<FVoiceChatDeviceInfo> OutputDevices;
 	const VivoxClientApi::AudioDeviceId* AudioDevices = nullptr;
 	int NumAudioDevices = 0;
 	VivoxClientConnection.GetAvailableAudioOutputDevices(LoginSession.AccountName, AudioDevices, NumAudioDevices);
@@ -251,18 +253,20 @@ TArray<FString> FVivoxVoiceChatUser::GetAvailableOutputDevices() const
 	{
 		for (int DeviceIndex = 0; DeviceIndex < NumAudioDevices; ++DeviceIndex)
 		{
-			OutputDevices.Add(UTF8_TO_TCHAR(AudioDevices[DeviceIndex].GetAudioDeviceDisplayName()));
+			FVoiceChatDeviceInfo& DeviceInfo = OutputDevices.Emplace_GetRef();
+			DeviceInfo.DisplayName = UTF8_TO_TCHAR(AudioDevices[DeviceIndex].GetAudioDeviceDisplayName());
+			DeviceInfo.Id = UTF8_TO_TCHAR(AudioDevices[DeviceIndex].GetAudioDeviceId());
 		}
 	}
 	return OutputDevices;
 }
 
-void FVivoxVoiceChatUser::SetInputDevice(const FString& InputDevice)
+void FVivoxVoiceChatUser::SetInputDeviceId(const FString& InputDeviceId)
 {
-	VIVOXVOICECHATUSER_LOG(Verbose, TEXT("SetInputDevice %s"), *InputDevice);
+	VIVOXVOICECHATUSER_LOG(Verbose, TEXT("SetInputDeviceId %s"), *InputDeviceId);
 
 	bool bDeviceFound = false;
-	if (!InputDevice.IsEmpty())
+	if (!InputDeviceId.IsEmpty())
 	{
 		const VivoxClientApi::AudioDeviceId* AudioDevices = nullptr;
 		int NumAudioDevices = 0;
@@ -272,9 +276,8 @@ void FVivoxVoiceChatUser::SetInputDevice(const FString& InputDevice)
 			for (int DeviceIndex = 0; DeviceIndex < NumAudioDevices; ++DeviceIndex)
 			{
 				const VivoxClientApi::AudioDeviceId& DeviceId = AudioDevices[DeviceIndex];
-				if (InputDevice == UTF8_TO_TCHAR(DeviceId.GetAudioDeviceDisplayName()))
+				if (InputDeviceId == UTF8_TO_TCHAR(DeviceId.GetAudioDeviceId()))
 				{
-
 					AudioInputOptions.DevicePolicy.SetSpecificAudioDevice(DeviceId);
 					bDeviceFound = true;
 					break;
@@ -291,12 +294,12 @@ void FVivoxVoiceChatUser::SetInputDevice(const FString& InputDevice)
 	ApplyAudioInputDevicePolicy();
 }
 
-void FVivoxVoiceChatUser::SetOutputDevice(const FString& OutputDevice)
+void FVivoxVoiceChatUser::SetOutputDeviceId(const FString& OutputDeviceId)
 {
-	VIVOXVOICECHATUSER_LOG(Verbose, TEXT("SetOutputDevice %s"), *OutputDevice);
+	VIVOXVOICECHATUSER_LOG(Verbose, TEXT("SetOutputDeviceId %s"), *OutputDeviceId);
 
 	bool bDeviceFound = false;
-	if (!OutputDevice.IsEmpty())
+	if (!OutputDeviceId.IsEmpty())
 	{
 		const VivoxClientApi::AudioDeviceId* AudioDevices = nullptr;
 		int NumAudioDevices = 0;
@@ -306,7 +309,7 @@ void FVivoxVoiceChatUser::SetOutputDevice(const FString& OutputDevice)
 			for (int DeviceIndex = 0; DeviceIndex < NumAudioDevices; ++DeviceIndex)
 			{
 				const VivoxClientApi::AudioDeviceId& DeviceId = AudioDevices[DeviceIndex];
-				if (OutputDevice == UTF8_TO_TCHAR(DeviceId.GetAudioDeviceDisplayName()))
+				if (OutputDeviceId == UTF8_TO_TCHAR(DeviceId.GetAudioDeviceId()))
 				{
 					AudioOutputOptions.DevicePolicy.SetSpecificAudioDevice(DeviceId);
 					bDeviceFound = true;
@@ -324,40 +327,56 @@ void FVivoxVoiceChatUser::SetOutputDevice(const FString& OutputDevice)
 	ApplyAudioOutputDevicePolicy();
 }
 
-FString FVivoxVoiceChatUser::GetInputDevice() const
+FVoiceChatDeviceInfo FVivoxVoiceChatUser::GetInputDeviceInfo() const
 {
+	FVoiceChatDeviceInfo DeviceInfo;
 	if (AudioInputOptions.DevicePolicy.GetAudioDevicePolicy() == VivoxClientApi::AudioDevicePolicy::vx_audio_device_policy_specific_device)
 	{
-		return UTF8_TO_TCHAR(AudioInputOptions.DevicePolicy.GetSpecificAudioDevice().GetAudioDeviceDisplayName());
+		VivoxClientApi::AudioDeviceId AudioDevice = AudioInputOptions.DevicePolicy.GetSpecificAudioDevice();
+		DeviceInfo.DisplayName = UTF8_TO_TCHAR(AudioDevice.GetAudioDeviceDisplayName());
+		DeviceInfo.Id = UTF8_TO_TCHAR(AudioDevice.GetAudioDeviceId());
 	}
 	else
 	{
-		return GetDefaultInputDevice();
+		DeviceInfo = GetDefaultInputDeviceInfo();
 	}
+	return DeviceInfo;
 }
 
-FString FVivoxVoiceChatUser::GetOutputDevice() const
+FVoiceChatDeviceInfo FVivoxVoiceChatUser::GetOutputDeviceInfo() const
 {
+	FVoiceChatDeviceInfo DeviceInfo;
 	if (AudioOutputOptions.DevicePolicy.GetAudioDevicePolicy() == VivoxClientApi::AudioDevicePolicy::vx_audio_device_policy_specific_device)
 	{
-		return UTF8_TO_TCHAR(AudioOutputOptions.DevicePolicy.GetSpecificAudioDevice().GetAudioDeviceDisplayName());
+		VivoxClientApi::AudioDeviceId AudioDevice = AudioOutputOptions.DevicePolicy.GetSpecificAudioDevice();
+		DeviceInfo.DisplayName = UTF8_TO_TCHAR(AudioDevice.GetAudioDeviceDisplayName());
+		DeviceInfo.Id = UTF8_TO_TCHAR(AudioDevice.GetAudioDeviceId());
 	}
 	else
 	{
-		return GetDefaultOutputDevice();
+		DeviceInfo = GetDefaultOutputDeviceInfo();
 	}
+	return DeviceInfo;
 }
 
-FString FVivoxVoiceChatUser::GetDefaultInputDevice() const
+FVoiceChatDeviceInfo FVivoxVoiceChatUser::GetDefaultInputDeviceInfo() const
 {
 	VivoxClientApi::AudioDeviceId DeviceId = VivoxClientConnection.GetOperatingSystemChosenAudioInputDevice(LoginSession.AccountName);
-	return UTF8_TO_TCHAR(DeviceId.GetAudioDeviceDisplayName());
+
+	FVoiceChatDeviceInfo DeviceInfo;
+	DeviceInfo.DisplayName = UTF8_TO_TCHAR(DeviceId.GetAudioDeviceDisplayName());
+	DeviceInfo.Id = UTF8_TO_TCHAR(DeviceId.GetAudioDeviceId());
+	return DeviceInfo;
 }
 
-FString FVivoxVoiceChatUser::GetDefaultOutputDevice() const
+FVoiceChatDeviceInfo FVivoxVoiceChatUser::GetDefaultOutputDeviceInfo() const
 {
 	VivoxClientApi::AudioDeviceId DeviceId = VivoxClientConnection.GetOperatingSystemChosenAudioOutputDevice(LoginSession.AccountName);
-	return UTF8_TO_TCHAR(DeviceId.GetAudioDeviceDisplayName());
+	
+	FVoiceChatDeviceInfo DeviceInfo;
+	DeviceInfo.DisplayName = UTF8_TO_TCHAR(DeviceId.GetAudioDeviceDisplayName());
+	DeviceInfo.Id = UTF8_TO_TCHAR(DeviceId.GetAudioDeviceId());
+	return DeviceInfo;
 }
 
 void FVivoxVoiceChatUser::Connect(const FOnVoiceChatConnectCompleteDelegate& Delegate)
@@ -1004,42 +1023,42 @@ bool FVivoxVoiceChatUser::Exec(UWorld* InWorld, const TCHAR* Cmd, FOutputDevice&
 		if (IsInitialized())
 		{
 			VIVOX_EXEC_LOG(TEXT("  Input Devices: muted:%s volume:%.2f"), *LexToString(GetAudioInputDeviceMuted()), GetAudioInputVolume());
-			const FString InputDevice = GetInputDevice();
-			const FString DefaultInputDevice = GetDefaultInputDevice();
+			const FVoiceChatDeviceInfo InputDevice = GetInputDeviceInfo();
+			const FVoiceChatDeviceInfo DefaultInputDevice = GetDefaultInputDeviceInfo();
 			if (InputDevice == DefaultInputDevice)
 			{
-				VIVOX_EXEC_LOG(TEXT("    [%s] (Selected) (Default)"), *DefaultInputDevice);
+				VIVOX_EXEC_LOG(TEXT("    [%s] (Selected) (Default)"), *LexToString(DefaultInputDevice));
 			}
 			else
 			{
-				VIVOX_EXEC_LOG(TEXT("    [%s] (Selected)"), *InputDevice);
-				VIVOX_EXEC_LOG(TEXT("    [%s] (Default)"), *DefaultInputDevice);
+				VIVOX_EXEC_LOG(TEXT("    [%s] (Selected)"), *LexToString(InputDevice));
+				VIVOX_EXEC_LOG(TEXT("    [%s] (Default)"), *LexToString(DefaultInputDevice));
 			}
-			for (const FString& Device : GetAvailableInputDevices())
+			for (const FVoiceChatDeviceInfo& Device : GetAvailableInputDeviceInfos())
 			{
 				if (Device != DefaultInputDevice && Device != InputDevice)
 				{
-					VIVOX_EXEC_LOG(TEXT("    [%s]"), *Device);
+					VIVOX_EXEC_LOG(TEXT("    [%s]"), *LexToString(Device));
 				}
 			}
 
 			VIVOX_EXEC_LOG(TEXT("  Output Devices: muted:%s volume:%.2f"), *LexToString(GetAudioOutputDeviceMuted()), GetAudioOutputVolume());
-			const FString OutputDevice = GetOutputDevice();
-			const FString DefaultOutputDevice = GetDefaultOutputDevice();
+			const FVoiceChatDeviceInfo OutputDevice = GetOutputDeviceInfo();
+			const FVoiceChatDeviceInfo DefaultOutputDevice = GetDefaultOutputDeviceInfo();
 			if (OutputDevice == DefaultOutputDevice)
 			{
-				VIVOX_EXEC_LOG(TEXT("    [%s] (Selected) (Default)"), *DefaultOutputDevice);
+				VIVOX_EXEC_LOG(TEXT("    [%s] (Selected) (Default)"), *LexToString(DefaultOutputDevice));
 			}
 			else
 			{
-				VIVOX_EXEC_LOG(TEXT("    [%s] (Selected)"), *OutputDevice);
-				VIVOX_EXEC_LOG(TEXT("    [%s] (Default)"), *DefaultOutputDevice);
+				VIVOX_EXEC_LOG(TEXT("    [%s] (Selected)"), *LexToString(OutputDevice));
+				VIVOX_EXEC_LOG(TEXT("    [%s] (Default)"), *LexToString(DefaultOutputDevice));
 			}
-			for (const FString& Device : GetAvailableOutputDevices())
+			for (const FVoiceChatDeviceInfo& Device : GetAvailableOutputDeviceInfos())
 			{
 				if (Device != DefaultOutputDevice && Device != OutputDevice)
 				{
-					VIVOX_EXEC_LOG(TEXT("    [%s]"), *Device);
+					VIVOX_EXEC_LOG(TEXT("    [%s]"), *LexToString(Device));
 				}
 			}
 
@@ -1111,26 +1130,25 @@ bool FVivoxVoiceChatUser::Exec(UWorld* InWorld, const TCHAR* Cmd, FOutputDevice&
 		}
 		else if (FParse::Command(&Cmd, TEXT("LISTDEVICES")))
 		{
-			TArray<FString> Devices = GetAvailableInputDevices();
 			VIVOX_EXEC_LOG(TEXT("Input Devices:"));
-			for (const FString& Device : Devices)
+			for (const FVoiceChatDeviceInfo& DeviceInfo : GetAvailableInputDeviceInfos())
 			{
-				VIVOX_EXEC_LOG(TEXT("  %s"), *Device);
+				VIVOX_EXEC_LOG(TEXT("  %s"), *LexToString(DeviceInfo));
 			}
 			return true;
 		}
 		else if (FParse::Command(&Cmd, TEXT("SETDEVICE")))
 		{
-			FString Device;
-			if (FParse::Token(Cmd, Device, false))
+			FString DeviceId;
+			if (FParse::Token(Cmd, DeviceId, false))
 			{
-				SetInputDevice(Device);
+				SetInputDeviceId(DeviceId);
 				return true;
 			}
 		}
 		else if (FParse::Command(&Cmd, TEXT("SETDEFAULTDEVICE")))
 		{
-			SetInputDevice(FString());
+			SetInputDeviceId(FString());
 			return true;
 		}
 	}
@@ -1157,26 +1175,25 @@ bool FVivoxVoiceChatUser::Exec(UWorld* InWorld, const TCHAR* Cmd, FOutputDevice&
 		}
 		else if (FParse::Command(&Cmd, TEXT("LISTDEVICES")))
 		{
-			TArray<FString> Devices = GetAvailableOutputDevices();
 			VIVOX_EXEC_LOG(TEXT("Output Devices:"));
-			for (const FString& Device : Devices)
+			for (const FVoiceChatDeviceInfo& DeviceInfo : GetAvailableOutputDeviceInfos())
 			{
-				VIVOX_EXEC_LOG(TEXT("  %s"), *Device);
+				VIVOX_EXEC_LOG(TEXT("  %s"), *LexToString(DeviceInfo));
 			}
 			return true;
 		}
 		else if (FParse::Command(&Cmd, TEXT("SETDEVICE")))
 		{
-			FString Device;
-			if (FParse::Token(Cmd, Device, false))
+			FString DeviceId;
+			if (FParse::Token(Cmd, DeviceId, false))
 			{
-				SetOutputDevice(Device);
+				SetOutputDeviceId(DeviceId);
 				return true;
 			}
 		}
 		else if (FParse::Command(&Cmd, TEXT("SETDEFAULTDEVICE")))
 		{
-			SetOutputDevice(FString());
+			SetOutputDeviceId(FString());
 			return true;
 		}
 	}
@@ -2463,14 +2480,14 @@ bool FVivoxVoiceChat::GetAudioOutputDeviceMuted() const
 	return GetVoiceChatUser().GetAudioOutputDeviceMuted();
 }
 
-TArray<FString> FVivoxVoiceChat::GetAvailableInputDevices() const
+TArray<FVoiceChatDeviceInfo> FVivoxVoiceChat::GetAvailableInputDeviceInfos() const
 {
-	return GetVoiceChatUser().GetAvailableInputDevices();
+	return GetVoiceChatUser().GetAvailableInputDeviceInfos();
 }
 
-TArray<FString> FVivoxVoiceChat::GetAvailableOutputDevices() const
+TArray<FVoiceChatDeviceInfo> FVivoxVoiceChat::GetAvailableOutputDeviceInfos() const
 {
-	return GetVoiceChatUser().GetAvailableOutputDevices();
+	return GetVoiceChatUser().GetAvailableOutputDeviceInfos();
 }
 
 FOnVoiceChatAvailableAudioDevicesChangedDelegate& FVivoxVoiceChat::OnVoiceChatAvailableAudioDevicesChanged()
@@ -2478,34 +2495,34 @@ FOnVoiceChatAvailableAudioDevicesChangedDelegate& FVivoxVoiceChat::OnVoiceChatAv
 	return GetVoiceChatUser().OnVoiceChatAvailableAudioDevicesChanged();
 }
 
-void FVivoxVoiceChat::SetInputDevice(const FString& InputDevice)
+void FVivoxVoiceChat::SetInputDeviceId(const FString& InputDeviceId)
 {
-	GetVoiceChatUser().SetInputDevice(InputDevice);
+	GetVoiceChatUser().SetInputDeviceId(InputDeviceId);
 }
 
-void FVivoxVoiceChat::SetOutputDevice(const FString& OutputDevice)
+void FVivoxVoiceChat::SetOutputDeviceId(const FString& OutputDeviceId)
 {
-	GetVoiceChatUser().SetOutputDevice(OutputDevice);
+	GetVoiceChatUser().SetOutputDeviceId(OutputDeviceId);
 }
 
-FString FVivoxVoiceChat::GetInputDevice() const
+FVoiceChatDeviceInfo FVivoxVoiceChat::GetInputDeviceInfo() const
 {
-	return GetVoiceChatUser().GetInputDevice();
+	return GetVoiceChatUser().GetInputDeviceInfo();
 }
 
-FString FVivoxVoiceChat::GetOutputDevice() const
+FVoiceChatDeviceInfo FVivoxVoiceChat::GetOutputDeviceInfo() const
 {
-	return GetVoiceChatUser().GetOutputDevice();
+	return GetVoiceChatUser().GetOutputDeviceInfo();
 }
 
-FString FVivoxVoiceChat::GetDefaultInputDevice() const
+FVoiceChatDeviceInfo FVivoxVoiceChat::GetDefaultInputDeviceInfo() const
 {
-	return GetVoiceChatUser().GetDefaultInputDevice();
+	return GetVoiceChatUser().GetDefaultInputDeviceInfo();
 }
 
-FString FVivoxVoiceChat::GetDefaultOutputDevice() const
+FVoiceChatDeviceInfo FVivoxVoiceChat::GetDefaultOutputDeviceInfo() const
 {
-	return GetVoiceChatUser().GetDefaultOutputDevice();
+	return GetVoiceChatUser().GetDefaultOutputDeviceInfo();
 }
 
 void FVivoxVoiceChat::Connect(const FOnVoiceChatConnectCompleteDelegate& Delegate)
diff --git a/Engine/Plugins/Online/VoiceChat/VivoxVoiceChat/Source/Public/VivoxVoiceChat.h b/Engine/Plugins/Online/VoiceChat/VivoxVoiceChat/Source/Public/VivoxVoiceChat.h
--- a/Engine/Plugins/Online/VoiceChat/VivoxVoiceChat/Source/Public/VivoxVoiceChat.h
+++ b/Engine/Plugins/Online/VoiceChat/VivoxVoiceChat/Source/Public/VivoxVoiceChat.h
@@ -55,15 +55,15 @@ public:
 	virtual void SetAudioOutputDeviceMuted(bool bIsMuted) override;
 	virtual bool GetAudioInputDeviceMuted() const override;
 	virtual bool GetAudioOutputDeviceMuted() const override;
-	virtual TArray<FString> GetAvailableInputDevices() const override;
-	virtual TArray<FString> GetAvailableOutputDevices() const override;
+	virtual TArray<FVoiceChatDeviceInfo> GetAvailableInputDeviceInfos() const override;
+	virtual TArray<FVoiceChatDeviceInfo> GetAvailableOutputDeviceInfos() const override;
 	virtual FOnVoiceChatAvailableAudioDevicesChangedDelegate& OnVoiceChatAvailableAudioDevicesChanged() override { return OnVoiceChatAvailableAudioDevicesChangedDelegate; }
-	virtual void SetInputDevice(const FString& InputDevice) override;
-	virtual void SetOutputDevice(const FString& OutputDevice) override;
-	virtual FString GetInputDevice() const override;
-	virtual FString GetOutputDevice() const override;
-	virtual FString GetDefaultInputDevice() const override;
-	virtual FString GetDefaultOutputDevice() const override;
+	virtual void SetInputDeviceId(const FString& InputDevice) override;
+	virtual void SetOutputDeviceId(const FString& OutputDevice) override;
+	virtual FVoiceChatDeviceInfo GetInputDeviceInfo() const override;
+	virtual FVoiceChatDeviceInfo GetOutputDeviceInfo() const override;
+	virtual FVoiceChatDeviceInfo GetDefaultInputDeviceInfo() const override;
+	virtual FVoiceChatDeviceInfo GetDefaultOutputDeviceInfo() const override;
 	virtual void Connect(const FOnVoiceChatConnectCompleteDelegate& Delegate) override;
 	virtual void Disconnect(const FOnVoiceChatDisconnectCompleteDelegate& Delegate) override;
 	virtual bool IsConnecting() const override;
@@ -294,15 +294,15 @@ public:
 	virtual void SetAudioOutputDeviceMuted(bool bIsMuted) override;
 	virtual bool GetAudioInputDeviceMuted() const override;
 	virtual bool GetAudioOutputDeviceMuted() const override;
-	virtual TArray<FString> GetAvailableInputDevices() const override;
-	virtual TArray<FString> GetAvailableOutputDevices() const override;
+	virtual TArray<FVoiceChatDeviceInfo> GetAvailableInputDeviceInfos() const override;
+	virtual TArray<FVoiceChatDeviceInfo> GetAvailableOutputDeviceInfos() const override;
 	virtual FOnVoiceChatAvailableAudioDevicesChangedDelegate& OnVoiceChatAvailableAudioDevicesChanged() override;
-	virtual void SetInputDevice(const FString& InputDevice) override;
-	virtual void SetOutputDevice(const FString& OutputDevice) override;
-	virtual FString GetInputDevice() const override;
-	virtual FString GetOutputDevice() const override;
-	virtual FString GetDefaultInputDevice() const override;
-	virtual FString GetDefaultOutputDevice() const override;
+	virtual void SetInputDeviceId(const FString& InputDeviceId) override;
+	virtual void SetOutputDeviceId(const FString& OutputDeviceId) override;
+	virtual FVoiceChatDeviceInfo GetInputDeviceInfo() const override;
+	virtual FVoiceChatDeviceInfo GetOutputDeviceInfo() const override;
+	virtual FVoiceChatDeviceInfo GetDefaultInputDeviceInfo() const override;
+	virtual FVoiceChatDeviceInfo GetDefaultOutputDeviceInfo() const override;
 	virtual void Connect(const FOnVoiceChatConnectCompleteDelegate& Delegate) override;
 	virtual void Disconnect(const FOnVoiceChatDisconnectCompleteDelegate& Delegate) override;
 	virtual bool IsConnecting() const override;
diff --git a/Engine/Plugins/Online/VoiceChat/VoiceChat/Source/Public/VoiceChat.h b/Engine/Plugins/Online/VoiceChat/VoiceChat/Source/Public/VoiceChat.h
--- a/Engine/Plugins/Online/VoiceChat/VoiceChat/Source/Public/VoiceChat.h
+++ b/Engine/Plugins/Online/VoiceChat/VoiceChat/Source/Public/VoiceChat.h
@@ -68,6 +68,22 @@ struct FVoiceChatCallStats
 	int PacketsNumLost;
 };
 
+struct FVoiceChatDeviceInfo
+{
+	/** The display name for the device */
+	FString DisplayName;
+	/** The unique id for the device */
+	FString Id;
+
+	inline bool operator==(const FVoiceChatDeviceInfo& Other) const { return Id == Other.Id; }
+	inline bool operator!=(const FVoiceChatDeviceInfo& Other) const { return !operator==(Other); }
+};
+
+inline FString LexToString(const FVoiceChatDeviceInfo& DeviceInfo)
+{
+	return FString::Printf(TEXT("DisplayName=[%s] Id=[%s]"), *DeviceInfo.DisplayName, *DeviceInfo.Id);
+}
+
 DECLARE_DELEGATE_OneParam(FOnVoiceChatConnectCompleteDelegate, const FVoiceChatResult& /* Result */);
 DECLARE_DELEGATE_OneParam(FOnVoiceChatDisconnectCompleteDelegate, const FVoiceChatResult& /* Result */);
 DECLARE_DELEGATE_TwoParams(FOnVoiceChatLoginCompleteDelegate, const FString& /* PlayerName */, const FVoiceChatResult& /* Result */);
@@ -173,19 +189,51 @@ public:
 	 */
 	virtual bool GetAudioOutputDeviceMuted() const = 0;
 
+	/**
+	 * Get a list of available audio input device names
+	 *
+	 * @return Array of audio input devices names
+	 */
+	UE_DEPRECATED(4.26, "GetAvailableInputDevices is deprecated. Please use GetAvailableInputDeviceInfos instead!")
+	TArray<FString> GetAvailableInputDevices() const
+	{
+		TArray<FString> DeviceNames;
+		for (const FVoiceChatDeviceInfo& DeviceInfo : GetAvailableInputDeviceInfos())
+		{
+			DeviceNames.Emplace(DeviceInfo.DisplayName);
+		}
+		return DeviceNames;
+	}
+
 	/**
 	 * Get a list of available audio input devices
 	 *
 	 * @return Array of audio input devices
 	 */
-	virtual TArray<FString> GetAvailableInputDevices() const = 0;
+	virtual TArray<FVoiceChatDeviceInfo> GetAvailableInputDeviceInfos() const = 0;
+
+	/**
+	 * Get a list of available audio output device names
+	 *
+	 * @return Array of audio output device names
+	 */
+	UE_DEPRECATED(4.26, "GetAvailableOutputDevices is deprecated. Please use GetAvailableOutputDeviceInfos instead!")
+	TArray<FString> GetAvailableOutputDevices() const
+	{
+		TArray<FString> DeviceNames;
+		for (const FVoiceChatDeviceInfo& DeviceInfo : GetAvailableOutputDeviceInfos())
+		{
+			DeviceNames.Emplace(DeviceInfo.DisplayName);
+		}
+		return DeviceNames;
+	}
 
 	/**
 	 * Get a list of available audio output devices
 	 *
 	 * @return Array of audio output devices
 	 */
-	virtual TArray<FString> GetAvailableOutputDevices() const = 0;
+	virtual TArray<FVoiceChatDeviceInfo> GetAvailableOutputDeviceInfos() const = 0;
 
 	/**
 	 * Delegate broadcast whenever the available audio devices change. Any cached values from GetAvailableInputDevices or GetAvailableOutputDevices should be discarded and requeried
@@ -197,34 +245,110 @@ public:
 	 *
 	 * @param InputDevice String from GetAvailableInputDevice. If empty, will use the default input device
 	 */
-	virtual void SetInputDevice(const FString& InputDevice) = 0;
+	UE_DEPRECATED(4.26, "SetInputDevice is deprecated. Please use SetInputDeviceId instead!")
+	void SetInputDevice(const FString& InputDeviceName)
+	{
+		FString InputDeviceId = InputDeviceName;
+		for (const FVoiceChatDeviceInfo& DeviceInfo : GetAvailableInputDeviceInfos())
+		{
+			if (InputDeviceName == DeviceInfo.DisplayName)
+			{
+				InputDeviceId = DeviceInfo.Id;
+				break;
+			}
+		}
+		SetInputDeviceId(InputDeviceId);
+	}
+
+	/**
+	 * Set the audio input device to use
+	 *
+	 * @param InputDeviceId Device Id from GetAvailableInputDeviceInfos. If empty, will use the default input device
+	 */
+	virtual void SetInputDeviceId(const FString& InputDeviceId) = 0;
 
 	/**
 	 * Set the audio output device to use
 	 *
 	 * @param OutputDevice String from GetAvailableOutputDevice. If empty, will use the default output device
 	 */
-	virtual void SetOutputDevice(const FString& OutputDevice) = 0;
+	UE_DEPRECATED(4.26, "SetOutputDevice is deprecated. Please use SetOutputDeviceId instead!")
+	void SetOutputDevice(const FString& OutputDeviceName)
+	{
+		FString OutputDeviceId = OutputDeviceName;
+		for (const FVoiceChatDeviceInfo& DeviceInfo : GetAvailableOutputDeviceInfos())
+		{
+			if (OutputDeviceName == DeviceInfo.DisplayName)
+			{
+				OutputDeviceId = DeviceInfo.Id;
+				break;
+			}
+		}
+		SetOutputDeviceId(OutputDeviceId);
+	}
 
 	/**
-	 * Get the audio input device that is being used
+	 * Set the audio output device to use
+	 *
+	 * @param OutputDeviceId Device Id from GetAvailableOutputDeviceInfos. If empty, will use the default output device
 	 */
-	virtual FString GetInputDevice() const = 0;
+	virtual void SetOutputDeviceId(const FString& InputDeviceId) = 0;
 
 	/**
-	 * Get the audio output device that is being used
+	 * Get the display name of the audio input device that is being used
 	 */
-	virtual FString GetOutputDevice() const = 0;
+	UE_DEPRECATED(4.26, "GetInputDevice is deprecated. Please use GetInputDeviceInfo instead!")
+	FString GetInputDevice() const
+	{
+		return GetInputDeviceInfo().DisplayName;
+	}
+
+	/**
+	 * Get the device info of the audio input device that is being used
+	 */
+	virtual FVoiceChatDeviceInfo GetInputDeviceInfo() const = 0;
 
 	/**
-	 * Get the default audio input device
+	 * Get the display name of the audio output device that is being used
 	 */
-	virtual FString GetDefaultInputDevice() const = 0;
+	UE_DEPRECATED(4.26, "GetOutputDevice is deprecated. Please use GetOutputDeviceInfo instead!")
+	FString GetOutputDevice() const
+	{
+		return GetOutputDeviceInfo().DisplayName;
+	}
+
+	/**
+	 * Get the device info of the audio output device that is being used
+	 */
+	virtual FVoiceChatDeviceInfo GetOutputDeviceInfo() const = 0;
+
+	/**
+	 * Get the display name of the default audio input device
+	 */
+	UE_DEPRECATED(4.26, "GetDefaultInputDevice is deprecated. Please use GetDefaultInputDeviceInfo instead!")
+	FString GetDefaultInputDevice() const
+	{
+		return GetDefaultInputDeviceInfo().DisplayName;
+	};
+
+	/**
+	 * Get the device info of the default audio input device
+	 */
+	virtual FVoiceChatDeviceInfo GetDefaultInputDeviceInfo() const = 0;
+
+	/**
+	 * Get the display name of the default audio output device
+	 */
+	UE_DEPRECATED(4.26, "GetDefaultOutputDevice is deprecated. Please use GetDefaultOutputDeviceInfo instead!")
+	FString GetDefaultOutputDevice() const
+	{
+		return GetDefaultOutputDeviceInfo().DisplayName;
+	}
 
 	/**
-	 * Get the default audio output device
+	 * Get the device info of the default audio output device
 	 */
-	virtual FString GetDefaultOutputDevice() const = 0;
+	virtual FVoiceChatDeviceInfo GetDefaultOutputDeviceInfo() const = 0;
 
 	/**
 	 * Connect to a voice server
