- Move ConnectionState out of VoiceChatUser and back into VoiceChat

diff --git a/Engine/Plugins/Online/VoiceChat/VivoxVoiceChat/Source/Private/VivoxVoiceChat.cpp b/Engine/Plugins/Online/VoiceChat/VivoxVoiceChat/Source/Private/VivoxVoiceChat.cpp
--- a/Engine/Plugins/Online/VoiceChat/VivoxVoiceChat/Source/Private/VivoxVoiceChat.cpp
+++ b/Engine/Plugins/Online/VoiceChat/VivoxVoiceChat/Source/Private/VivoxVoiceChat.cpp
@@ -379,41 +379,6 @@ FVoiceChatDeviceInfo FVivoxVoiceChatUser::GetDefaultOutputDeviceInfo() const
 	return DeviceInfo;
 }
 
-void FVivoxVoiceChatUser::Connect(const FOnVoiceChatConnectCompleteDelegate& Delegate)
-{
-	VivoxVoiceChat.Connect(Delegate);
-}
-
-void FVivoxVoiceChatUser::Disconnect(const FOnVoiceChatDisconnectCompleteDelegate& Delegate)
-{
-	VivoxVoiceChat.Disconnect(Delegate);
-}
-
-bool FVivoxVoiceChatUser::IsConnecting() const
-{
-	return VivoxVoiceChat.IsConnecting();
-}
-
-bool FVivoxVoiceChatUser::IsConnected() const
-{
-	return VivoxVoiceChat.IsConnected();
-}
-
-FOnVoiceChatConnectedDelegate& FVivoxVoiceChatUser::OnVoiceChatConnected()
-{
-	return VivoxVoiceChat.OnVoiceChatConnected();
-}
-
-FOnVoiceChatDisconnectedDelegate& FVivoxVoiceChatUser::OnVoiceChatDisconnected()
-{
-	return VivoxVoiceChat.OnVoiceChatDisconnected();
-}
-
-FOnVoiceChatReconnectedDelegate& FVivoxVoiceChatUser::OnVoiceChatReconnected()
-{
-	return VivoxVoiceChat.OnVoiceChatReconnected();
-}
-
 void FVivoxVoiceChatUser::Login(FPlatformUserId PlatformId, const FString& PlayerName, const FString& Credentials, const FOnVoiceChatLoginCompleteDelegate& Delegate)
 {
 	FVoiceChatResult Result = FVoiceChatResult::CreateSuccess();
@@ -1010,6 +975,11 @@ bool FVivoxVoiceChatUser::IsInitialized()
 	return VivoxVoiceChat.IsInitialized();
 }
 
+bool FVivoxVoiceChatUser::IsConnected()
+{
+	return VivoxVoiceChat.IsConnected();
+}
+
 bool FVivoxVoiceChatUser::Exec(UWorld* InWorld, const TCHAR* Cmd, FOutputDevice& Ar)
 {
 #if NO_LOGGING
@@ -1197,22 +1167,6 @@ bool FVivoxVoiceChatUser::Exec(UWorld* InWorld, const TCHAR* Cmd, FOutputDevice&
 			return true;
 		}
 	}
-	else if (FParse::Command(&Cmd, TEXT("CONNECT")))
-	{
-		Connect(FOnVoiceChatConnectCompleteDelegate::CreateLambda([](const FVoiceChatResult& Result)
-		{
-			UE_LOG(LogVivoxVoiceChat, Log, TEXT("VIVOX CONNECT result:%s"), *LexToString(Result));
-		}));
-		return true;
-	}
-	else if (FParse::Command(&Cmd, TEXT("DISCONNECT")))
-	{
-		Disconnect(FOnVoiceChatDisconnectCompleteDelegate::CreateLambda([](const FVoiceChatResult& Result)
-		{
-			UE_LOG(LogVivoxVoiceChat, Log, TEXT("VIVOX DISCONNECT result:%s"), *LexToString(Result));
-		}));
-		return true;
-	}
 	else if (FParse::Command(&Cmd, TEXT("LOGIN")))
 	{
 		FString PlayerName;
@@ -2325,6 +2279,22 @@ bool FVivoxVoiceChat::Exec(UWorld* InWorld, const TCHAR* Cmd, FOutputDevice& Ar)
 			Uninitialize();
 			return true;
 		}
+		else if (FParse::Command(&Cmd, TEXT("CONNECT")))
+		{
+			Connect(FOnVoiceChatConnectCompleteDelegate::CreateLambda([](const FVoiceChatResult& Result)
+			{
+				UE_LOG(LogVivoxVoiceChat, Log, TEXT("VIVOX CONNECT result:%s"), *LexToString(Result));
+			}));
+			return true;
+		}
+		else if (FParse::Command(&Cmd, TEXT("DISCONNECT")))
+		{
+			Disconnect(FOnVoiceChatDisconnectCompleteDelegate::CreateLambda([](const FVoiceChatResult& Result)
+			{
+				UE_LOG(LogVivoxVoiceChat, Log, TEXT("VIVOX DISCONNECT result:%s"), *LexToString(Result));
+			}));
+			return true;
+		}
 		else if (FParse::Command(&Cmd, TEXT("CREATEUSER")))
 		{
 			(void)CreateUser();
diff --git a/Engine/Plugins/Online/VoiceChat/VivoxVoiceChat/Source/Public/VivoxVoiceChat.h b/Engine/Plugins/Online/VoiceChat/VivoxVoiceChat/Source/Public/VivoxVoiceChat.h
--- a/Engine/Plugins/Online/VoiceChat/VivoxVoiceChat/Source/Public/VivoxVoiceChat.h
+++ b/Engine/Plugins/Online/VoiceChat/VivoxVoiceChat/Source/Public/VivoxVoiceChat.h
@@ -64,13 +64,6 @@ public:
 	virtual FVoiceChatDeviceInfo GetOutputDeviceInfo() const override;
 	virtual FVoiceChatDeviceInfo GetDefaultInputDeviceInfo() const override;
 	virtual FVoiceChatDeviceInfo GetDefaultOutputDeviceInfo() const override;
-	virtual void Connect(const FOnVoiceChatConnectCompleteDelegate& Delegate) override;
-	virtual void Disconnect(const FOnVoiceChatDisconnectCompleteDelegate& Delegate) override;
-	virtual bool IsConnecting() const override;
-	virtual bool IsConnected() const override;
-	virtual FOnVoiceChatConnectedDelegate& OnVoiceChatConnected() override;
-	virtual FOnVoiceChatDisconnectedDelegate& OnVoiceChatDisconnected() override;
-	virtual FOnVoiceChatReconnectedDelegate& OnVoiceChatReconnected() override;
 	virtual void Login(FPlatformUserId PlatformId, const FString& PlayerName, const FString& Credentials, const FOnVoiceChatLoginCompleteDelegate& Delegate) override;
 	virtual void Logout(const FOnVoiceChatLogoutCompleteDelegate& Delegate) override;
 	virtual bool IsLoggingIn() const override;
@@ -118,6 +111,7 @@ public:
 
 protected:
 	bool IsInitialized();
+	bool IsConnected();
 	bool Exec(UWorld* InWorld, const TCHAR* Cmd, FOutputDevice& Ar);
 	
 	friend class FVivoxVoiceChat;
@@ -281,6 +275,13 @@ public:
 	virtual bool Initialize() override;
 	virtual bool Uninitialize() override;
 	virtual bool IsInitialized() const override;
+	virtual void Connect(const FOnVoiceChatConnectCompleteDelegate& Delegate) override;
+	virtual void Disconnect(const FOnVoiceChatDisconnectCompleteDelegate& Delegate) override;
+	virtual bool IsConnecting() const override;
+	virtual bool IsConnected() const override;
+	virtual FOnVoiceChatConnectedDelegate& OnVoiceChatConnected() override { return OnVoiceChatConnectedDelegate; }
+	virtual FOnVoiceChatDisconnectedDelegate& OnVoiceChatDisconnected() override { return OnVoiceChatDisconnectedDelegate; }
+	virtual FOnVoiceChatReconnectedDelegate& OnVoiceChatReconnected() override { return OnVoiceChatReconnectedDelegate; }
 	virtual IVoiceChatUser* CreateUser() override;
 
 	// IVoiceChatUser
@@ -303,13 +304,6 @@ public:
 	virtual FVoiceChatDeviceInfo GetOutputDeviceInfo() const override;
 	virtual FVoiceChatDeviceInfo GetDefaultInputDeviceInfo() const override;
 	virtual FVoiceChatDeviceInfo GetDefaultOutputDeviceInfo() const override;
-	virtual void Connect(const FOnVoiceChatConnectCompleteDelegate& Delegate) override;
-	virtual void Disconnect(const FOnVoiceChatDisconnectCompleteDelegate& Delegate) override;
-	virtual bool IsConnecting() const override;
-	virtual bool IsConnected() const override;
-	virtual FOnVoiceChatConnectedDelegate& OnVoiceChatConnected() override { return OnVoiceChatConnectedDelegate; }
-	virtual FOnVoiceChatDisconnectedDelegate& OnVoiceChatDisconnected() override { return OnVoiceChatDisconnectedDelegate; }
-	virtual FOnVoiceChatReconnectedDelegate& OnVoiceChatReconnected() override { return OnVoiceChatReconnectedDelegate; }
 	virtual void Login(FPlatformUserId PlatformId, const FString& PlayerName, const FString& Credentials, const FOnVoiceChatLoginCompleteDelegate& Delegate) override;
 	virtual void Logout(const FOnVoiceChatLogoutCompleteDelegate& Delegate) override;
 	virtual bool IsLoggingIn() const override;
