- Reorder when we set callbacks to avoid possibility of race conditions

diff --git a/Engine/Plugins/Online/VoiceChat/VivoxVoiceChat/Source/Private/VivoxVoiceChat.cpp b/Engine/Plugins/Online/VoiceChat/VivoxVoiceChat/Source/Private/VivoxVoiceChat.cpp
--- a/Engine/Plugins/Online/VoiceChat/VivoxVoiceChat/Source/Private/VivoxVoiceChat.cpp
+++ b/Engine/Plugins/Online/VoiceChat/VivoxVoiceChat/Source/Private/VivoxVoiceChat.cpp
@@ -441,6 +441,8 @@ void FVivoxVoiceChatUser::Login(FPlatformUserId PlatformId, const FString& Playe
 	LoginSession.UserUri = UserUri;
 	LoginSession.State = FLoginSession::EState::LoggedOut;
 
+	OnVoiceChatLoginCompleteDelegate = Delegate;
+
 	VivoxClientApi::VCSStatus Status = VivoxClientConnection.Login(LoginSession.AccountName, TCHAR_TO_ANSI(*Credentials));
 	if (Status.IsError())
 	{
@@ -448,13 +450,11 @@ void FVivoxVoiceChatUser::Login(FPlatformUserId PlatformId, const FString& Playe
 		VIVOXVOICECHATUSER_LOG(Warning, TEXT("Login account:%s %s"), ANSI_TO_TCHAR(LoginSession.AccountName.ToString()), *LexToString(Result));
 
 		ClearLoginSession();
-		Delegate.ExecuteIfBound(PlayerName, Result);
+		TriggerCompletionDelegate(OnVoiceChatLoginCompleteDelegate, PlayerName, Result);
 		return;
 	}
 
 	LoginSession.State = FLoginSession::EState::LoggingIn;
-
-	OnVoiceChatLoginCompleteDelegate = Delegate;
 }
 
 void FVivoxVoiceChatUser::Logout(const FOnVoiceChatLogoutCompleteDelegate& Delegate)
@@ -482,17 +482,18 @@ void FVivoxVoiceChatUser::Logout(const FOnVoiceChatLogoutCompleteDelegate& Deleg
 		return;
 	}
 
+	OnVoiceChatLogoutCompleteDelegate = Delegate;
+
 	VivoxClientApi::AccountName AccountName = VivoxVoiceChat.CreateAccountName(LoginSession.PlayerName);
 	VivoxClientApi::VCSStatus Status = VivoxClientConnection.Logout(AccountName);
 	if (Status.IsError())
 	{
 		Result = ResultFromVivoxStatus(Status);
 		VIVOXVOICECHATUSER_LOG(Warning, TEXT("Logout %s"), *LexToString(Result));
-		Delegate.ExecuteIfBound(LoginSession.PlayerName, Result);
+		TriggerCompletionDelegate(OnVoiceChatLogoutCompleteDelegate, LoginSession.PlayerName, Result);
 		return;
 	}
 
-	OnVoiceChatLogoutCompleteDelegate = Delegate;
 	LoginSession.State = FLoginSession::EState::LoggingOut;
 }
 
@@ -612,18 +613,18 @@ void FVivoxVoiceChatUser::JoinChannel(const FString& ChannelName, const FString&
 	ChannelSession.ChannelName = ChannelName;
 	ChannelSession.ChannelType = ChannelType;
 	ChannelSession.ChannelUri = VivoxVoiceChat.CreateChannelUri(ChannelName, ChannelType, Channel3dProperties);
+	ChannelSession.JoinDelegate = Delegate;
 
 	VivoxClientApi::VCSStatus Status = VivoxClientConnection.JoinChannel(LoginSession.AccountName, ChannelSession.ChannelUri, TCHAR_TO_ANSI(*ChannelCredentials));
 	if (Status.IsError())
 	{
 		Result = ResultFromVivoxStatus(Status);
 		VIVOXVOICECHATUSER_LOG(Warning, TEXT("JoinChannel ChannelUri:%s %s"), ANSI_TO_TCHAR(ChannelSession.ChannelUri.ToString()), *LexToString(Result));
-		Delegate.ExecuteIfBound(ChannelName, Result);
+		TriggerCompletionDelegate(ChannelSession.JoinDelegate, ChannelName, Result);
 		return;
 	}
 
 	ChannelSession.State = FChannelSession::EState::Connecting;
-	ChannelSession.JoinDelegate = Delegate;
 }
 
 void FVivoxVoiceChatUser::LeaveChannel(const FString& ChannelName, const FOnVoiceChatChannelLeaveCompleteDelegate& Delegate)
@@ -660,17 +661,18 @@ void FVivoxVoiceChatUser::LeaveChannel(const FString& ChannelName, const FOnVoic
 		return;
 	}
 
+	ChannelSession.LeaveDelegate = Delegate;
+
 	VivoxClientApi::VCSStatus Status = VivoxClientConnection.LeaveChannel(LoginSession.AccountName, ChannelSession.ChannelUri);
 	if (Status.IsError())
 	{
 		Result = ResultFromVivoxStatus(Status);
 		VIVOXVOICECHATUSER_LOG(Warning, TEXT("LeaveChannel channel:%s %s"), ANSI_TO_TCHAR(ChannelSession.ChannelUri.ToString()), *LexToString(Result));
-		Delegate.ExecuteIfBound(ChannelName, Result);
+		TriggerCompletionDelegate(ChannelSession.LeaveDelegate, ChannelName, Result);
 		return;
 	}
 
 	ChannelSession.State = FChannelSession::EState::Disconnecting;
-	ChannelSession.LeaveDelegate = Delegate;
 }
 
 void FVivoxVoiceChatUser::Set3DPosition(const FString& ChannelName, const FVector& SpeakerPosition, const FVector& ListenerPosition, const FVector& ListenerForwardDirection, const FVector& ListenerUpDirection)
@@ -2560,11 +2562,12 @@ void FVivoxVoiceChat::Disconnect(const FOnVoiceChatDisconnectCompleteDelegate& D
 {
 	if (IsConnected())
 	{
+		OnVoiceChatDisconnectCompleteDelegates.Add(Delegate);
+
 		VivoxClientApi::Uri BackendUri(TCHAR_TO_ANSI(*VivoxServerUrl));
 		VivoxClientConnection.Disconnect(BackendUri);
 
 		ConnectionState = EConnectionState::Disconnecting;
-		OnVoiceChatDisconnectCompleteDelegates.Add(Delegate);
 	}
 	else
 	{
